Cameras: New camera shake previewer tool.

This is a panel that lets the user run camera shakes in editor mode.

diff --git a/Engine/Plugins/Cameras/CameraShakePreviewer/CameraShakePreviewer.uplugin b/Engine/Plugins/Cameras/CameraShakePreviewer/CameraShakePreviewer.uplugin
new file mode 100644
--- /dev/null
+++ b/Engine/Plugins/Cameras/CameraShakePreviewer/CameraShakePreviewer.uplugin
@@ -0,0 +1,25 @@
+{
+	"FileVersion" : 1,
+	"Version" : 1,
+	"VersionName" : "0.1",
+	"FriendlyName" : "Camera Shake Previewer",
+	"Description" : "Adds a new panel, accessible from the Level Editor, which lets the user preview camera shakes in editor viewports.",
+	"Category" : "Cameras",
+	"CreatedBy" : "Epic Games, Inc.",
+	"CreatedByURL" : "http://epicgames.com",
+	"DocsURL" : "",
+	"MarketplaceURL" : "",
+	"SupportURL" : "",
+	"EnabledByDefault" : true,
+	"CanContainContent" : false,
+	"IsBetaVersion" : true,
+	"Installed" : false,
+	"Modules" :
+	[
+		{
+			"Name" : "CameraShakePreviewer",
+			"Type" : "Editor",
+			"LoadingPhase" : "PostEngineInit"
+		}
+	]
+}
diff --git a/Engine/Plugins/Cameras/CameraShakePreviewer/Source/CameraShakePreviewer/CameraShakePreviewer.Build.cs b/Engine/Plugins/Cameras/CameraShakePreviewer/Source/CameraShakePreviewer/CameraShakePreviewer.Build.cs
new file mode 100644
--- /dev/null
+++ b/Engine/Plugins/Cameras/CameraShakePreviewer/Source/CameraShakePreviewer/CameraShakePreviewer.Build.cs
@@ -0,0 +1,41 @@
+// Copyright 1998-2019 Epic Games, Inc. All Rights Reserved.
+
+using UnrealBuildTool;
+
+public class CameraShakePreviewer : ModuleRules
+{
+	public CameraShakePreviewer(ReadOnlyTargetRules Target) : base(Target)
+	{
+		PrivateIncludePathModuleNames.AddRange(
+			new string[] {
+				"AssetTools",
+				"AssetRegistry",
+				"ContentBrowser",
+                "WorkspaceMenuStructure",
+            }
+		);
+
+		PrivateDependencyModuleNames.AddRange(
+			new string[] {
+				"ApplicationCore",
+				"SlateCore",
+				"AppFramework",
+				"EditorStyle",
+				"LevelEditor",
+				"PropertyEditor",
+				"Slate",
+				"ToolMenus",
+				"UnrealEd",
+			}
+		);
+
+		PublicDependencyModuleNames.AddRange(new string[] {
+				"Core",
+				"CoreUObject",
+				"Engine",
+				"InputCore",
+			}
+		);
+    }
+}
+
diff --git a/Engine/Plugins/Cameras/CameraShakePreviewer/Source/CameraShakePreviewer/Private/CameraShakePreviewerModule.cpp b/Engine/Plugins/Cameras/CameraShakePreviewer/Source/CameraShakePreviewer/Private/CameraShakePreviewerModule.cpp
new file mode 100644
--- /dev/null
+++ b/Engine/Plugins/Cameras/CameraShakePreviewer/Source/CameraShakePreviewer/Private/CameraShakePreviewerModule.cpp
@@ -0,0 +1,216 @@
+// Copyright 1998-2019 Epic Games, Inc. All Rights Reserved.
+
+#include "CameraShakePreviewerModule.h"
+#include "Framework/Commands/Commands.h"
+#include "Framework/MultiBox/MultiBoxBuilder.h"
+#include "Framework/MultiBox/MultiBoxExtender.h"
+#include "ILevelEditor.h"
+#include "ILevelViewport.h"
+#include "LevelEditor.h"
+#include "LevelEditorViewport.h"
+#include "Modules/ModuleManager.h"
+#include "SCameraShakePreviewer.h"
+#include "ToolMenus.h"
+#include "Widgets/DeclarativeSyntaxSupport.h"
+#include "Widgets/Docking/SDockTab.h"
+#include "Widgets/SWidget.h"
+#include "WorkspaceMenuStructure.h"
+#include "WorkspaceMenuStructureModule.h"
+
+#define LOCTEXT_NAMESPACE "CameraShakePreviewer"
+
+IMPLEMENT_MODULE(FCameraShakePreviewerModule, CameraShakePreviewer);
+
+static const FName LevelEditorModuleName("LevelEditor");
+static const FName LevelEditorCameraShakePreviewerTab("CameraShakePreviewer");
+
+/**
+ * Editor commands for the camera shake preview tool.
+ */
+class FCameraShakePreviewerCommands : public TCommands<FCameraShakePreviewerCommands>
+{
+public:
+	FCameraShakePreviewerCommands()
+		: TCommands<FCameraShakePreviewerCommands>(
+				TEXT("CameraShakePreviewer"),
+				LOCTEXT("CameraShakePreviewerContextDescription", "Camera Shake Previewer"),
+				TEXT("EditorViewport"),
+				FEditorStyle::GetStyleSetName())
+	{
+	}
+
+	static FCameraShakePreviewerCommands& Get()
+	{
+		return *(Instance.Pin());
+	}
+
+	virtual void RegisterCommands() override
+	{
+		UI_COMMAND(ToggleCameraShakesPreview, "Allow Camera Shakes", "If enabled, allows the camera shakes previewer panel to apply shakes to this viewport", EUserInterfaceActionType::ToggleButton, FInputChord());
+	}
+
+	TSharedPtr<FUICommandInfo> ToggleCameraShakesPreview;
+};
+
+
+void FCameraShakePreviewerModule::StartupModule()
+{
+	FCameraShakePreviewerCommands::Register();
+
+	if (ensure(FModuleManager::Get().IsModuleLoaded(LevelEditorModuleName)))
+	{
+		RegisterEditorTab();
+		RegisterViewportOptionMenuExtender();
+	}
+}
+
+void FCameraShakePreviewerModule::ShutdownModule()
+{
+	UnregisterViewportOptionMenuExtender();
+	UnregisterEditorTab();
+
+	FCameraShakePreviewerCommands::Unregister();
+}
+
+void FCameraShakePreviewerModule::RegisterEditorTab()
+{
+	FLevelEditorModule& LevelEditorModule = FModuleManager::LoadModuleChecked<FLevelEditorModule>(LevelEditorModuleName);
+
+	LevelEditorTabManagerChangedHandle = LevelEditorModule.OnTabManagerChanged().AddLambda([]()
+	{
+		// Add a new entry in the level editor's "Window" menu, which lets the user open the camera shake preview tool.
+		FLevelEditorModule& LevelEditorModule = FModuleManager::GetModuleChecked<FLevelEditorModule>(LevelEditorModuleName);
+		TSharedPtr<FTabManager> LevelEditorTabManager = LevelEditorModule.GetLevelEditorTabManager();
+
+		const IWorkspaceMenuStructure& MenuStructure = WorkspaceMenu::GetMenuStructure();
+		const FSlateIcon Icon(FEditorStyle::GetStyleSetName(), "LevelViewport.ToggleActorPilotCameraView");
+
+		LevelEditorTabManager->RegisterTabSpawner("CameraShakePreviewer", FOnSpawnTab::CreateStatic(&FCameraShakePreviewerModule::CreateCameraShakePreviewerTab))
+			.SetDisplayName(LOCTEXT("CameraShakePreviewer", "Camera Shake Previewer"))
+			.SetTooltipText(LOCTEXT("CameraShakePreviewerTooltipText", "Open the camera shake preview panel."))
+			.SetIcon(Icon)
+			.SetGroup(MenuStructure.GetLevelEditorCategory());
+	});
+}
+
+void FCameraShakePreviewerModule::UnregisterEditorTab()
+{
+	if (LevelEditorTabManagerChangedHandle.IsValid())
+	{
+		FLevelEditorModule& LevelEditorModule = FModuleManager::GetModuleChecked<FLevelEditorModule>(LevelEditorModuleName);
+		LevelEditorModule.OnTabManagerChanged().Remove(LevelEditorTabManagerChangedHandle);
+	}
+}
+
+void FCameraShakePreviewerModule::RegisterViewportOptionMenuExtender()
+{
+	// Register a callback for adding a "Show Camera Shakes" option in the viewport options menu.
+	FLevelEditorModule& LevelEditorModule = FModuleManager::LoadModuleChecked<FLevelEditorModule>(LevelEditorModuleName);
+
+	FLevelEditorModule::FLevelEditorMenuExtender Extender = FLevelEditorModule::FLevelEditorMenuExtender::CreateRaw(this, &FCameraShakePreviewerModule::OnExtendLevelViewportOptionMenu);
+	LevelEditorModule.GetAllLevelViewportOptionsMenuExtenders().Add(Extender);
+	ViewportOptionsMenuExtenderHandle = LevelEditorModule.GetAllLevelViewportOptionsMenuExtenders().Last().GetHandle();
+
+	GEditor->OnLevelViewportClientListChanged().AddRaw(this, &FCameraShakePreviewerModule::OnLevelViewportClientListChanged);
+
+	OnLevelViewportClientListChanged();
+}
+
+void FCameraShakePreviewerModule::UnregisterViewportOptionMenuExtender()
+{
+	FLevelEditorModule& LevelEditorModule = FModuleManager::GetModuleChecked<FLevelEditorModule>(LevelEditorModuleName);
+	FLevelEditorModule::FLevelEditorMenuExtender Extender = FLevelEditorModule::FLevelEditorMenuExtender::CreateRaw(this, &FCameraShakePreviewerModule::OnExtendLevelViewportOptionMenu);
+	LevelEditorModule.GetAllLevelViewportOptionsMenuExtenders().RemoveAll([=](const FLevelEditorModule::FLevelEditorMenuExtender& Extender) { return Extender.GetHandle() == ViewportOptionsMenuExtenderHandle; });
+
+	if (GEditor)
+	{
+		GEditor->OnLevelViewportClientListChanged().RemoveAll(this);
+	}
+}
+
+TSharedRef<FExtender> FCameraShakePreviewerModule::OnExtendLevelViewportOptionMenu(const TSharedRef<FUICommandList> CommandList)
+{
+	// Find the viewport for which we're opening this options menu.
+	FLevelEditorModule& LevelEditorModule = FModuleManager::LoadModuleChecked<FLevelEditorModule>(LevelEditorModuleName);
+	TSharedPtr<ILevelEditor> LevelEditor = LevelEditorModule.GetLevelEditorInstance().Pin();
+	TSharedPtr<ILevelViewport> ViewportInterface = LevelEditor->GetActiveViewportInterface();
+	FLevelEditorViewportClient* ViewportClient = &ViewportInterface->GetLevelViewportClient();
+	if (ViewportClient->ViewportType != ELevelViewportType::LVT_Perspective)
+	{
+		// If the current active viewport isn't perspective, let's not add any option for previewing camera shakes.
+		return MakeShared<FExtender>();
+	}
+
+	// Bind the "toggle camera shakes preview" action to some callbacks that will set the flag on the appropriate viewport.
+	FUIAction ToggleCameraShakesPreviewAction;
+	ToggleCameraShakesPreviewAction.ExecuteAction.BindLambda([this, ViewportClient]()
+	{
+		ToggleCameraShakesPreview(ViewportClient);
+	});
+	ToggleCameraShakesPreviewAction.GetActionCheckState.BindLambda([this, ViewportClient]() -> ECheckBoxState
+	{
+		return HasCameraShakesPreview(ViewportClient) ? ECheckBoxState::Checked : ECheckBoxState::Unchecked;
+	});
+
+	CommandList->MapAction(FCameraShakePreviewerCommands::Get().ToggleCameraShakesPreview, ToggleCameraShakesPreviewAction);
+
+	// Make the extender that adds the menu entry.
+	TSharedRef<FExtender> Extender = MakeShared<FExtender>();
+	Extender->AddMenuExtension(
+		"LevelViewportViewportOptions2",
+		EExtensionHook::After,
+		nullptr,
+		FMenuExtensionDelegate::CreateLambda([](FMenuBuilder& MenuBuilder)
+		{
+			MenuBuilder.AddMenuEntry(FCameraShakePreviewerCommands::Get().ToggleCameraShakesPreview);
+		})
+	);
+	return Extender;
+}
+
+void FCameraShakePreviewerModule::OnLevelViewportClientListChanged()
+{
+	const TArray<FLevelEditorViewportClient*> LevelViewportClients = GEditor->GetLevelViewportClients();
+	// Remove viewports that don't exist anymore.
+	for (auto It = ViewportInfos.CreateIterator(); It; ++It)
+	{
+		if (!LevelViewportClients.Contains(It.Key()))
+		{
+			ViewportInfos.Remove(It.Key());
+		}
+	}
+	// Add recently added viewports that we don't know about yet.
+	for (FLevelEditorViewportClient* LevelViewportClient : LevelViewportClients)
+	{
+		if (!ViewportInfos.Contains(LevelViewportClient))
+		{
+			ViewportInfos.Add(LevelViewportClient, FViewportInfo{ false });
+		}
+	}
+}
+
+void FCameraShakePreviewerModule::ToggleCameraShakesPreview(FLevelEditorViewportClient* ViewportClient)
+{
+	if (FViewportInfo* ViewportInfo = ViewportInfos.Find(ViewportClient))
+	{
+		ViewportInfo->bPreviewCameraShakes = !ViewportInfo->bPreviewCameraShakes;
+		OnTogglePreviewCameraShakes.Broadcast(FTogglePreviewCameraShakesParams{ ViewportClient, ViewportInfo->bPreviewCameraShakes });
+	}
+}
+
+bool FCameraShakePreviewerModule::HasCameraShakesPreview(FLevelEditorViewportClient* ViewportClient) const
+{
+	const FViewportInfo* ViewportInfo = ViewportInfos.Find(ViewportClient);
+	return ViewportInfo != nullptr && ViewportInfo->bPreviewCameraShakes;
+}
+
+TSharedRef<SDockTab> FCameraShakePreviewerModule::CreateCameraShakePreviewerTab(const FSpawnTabArgs& Args)
+{
+	return SNew(SDockTab)
+		.TabRole(ETabRole::PanelTab)
+		[
+			SNew(SCameraShakePreviewer)
+		];
+}
+
+#undef LOCTEXT_NAMESPACE
diff --git a/Engine/Plugins/Cameras/CameraShakePreviewer/Source/CameraShakePreviewer/Private/SCameraShakePreviewer.cpp b/Engine/Plugins/Cameras/CameraShakePreviewer/Source/CameraShakePreviewer/Private/SCameraShakePreviewer.cpp
new file mode 100644
--- /dev/null
+++ b/Engine/Plugins/Cameras/CameraShakePreviewer/Source/CameraShakePreviewer/Private/SCameraShakePreviewer.cpp
@@ -0,0 +1,683 @@
+// Copyright 1998-2019 Epic Games, Inc. All Rights Reserved.
+
+#include "SCameraShakePreviewer.h"
+#include "Camera/CameraModifier_CameraShake.h"
+#include "Camera/CameraShake.h"
+#include "Camera/CameraShakeSourceActor.h"
+#include "Camera/CameraShakeSourceComponent.h"
+#include "CameraShakePreviewerModule.h"
+#include "LevelEditor.h"
+#include "LevelEditorViewport.h"
+#include "Modules/ModuleManager.h"
+#include "Slate/SceneViewport.h"
+#include "Widgets/Input/SButton.h"
+#include "Widgets/Input/SCheckBox.h"
+#include "Widgets/Layout/SBox.h"
+#include "Widgets/Layout/SGridPanel.h"
+#include "Widgets/Layout/SWrapBox.h"
+#include "EditorStyleSet.h"
+#include "EditorDirectories.h"
+
+#define LOCTEXT_NAMESPACE "CameraShakePreviewer"
+
+/**
+ * Data struct for each entry in the panel's main list.
+ */
+struct FCameraShakeData
+	: public TSharedFromThis<FCameraShakeData>
+	, public FGCObject
+{
+	TSubclassOf<UCameraShake> ShakeClass;
+	UCameraShake* ShakeInstance = nullptr;
+	bool bIsPlaying = false;
+	bool bIsHidden = false;
+
+	TWeakObjectPtr<UCameraShakeSourceComponent> SourceComponent;
+
+	virtual void AddReferencedObjects(FReferenceCollector& Collector) override
+	{
+		Collector.AddReferencedObject(ShakeInstance);
+	}
+};
+
+/**
+ * The UI for each entry in the panel's main list.
+ */
+class SCameraShakeRow : public SMultiColumnTableRow<TSharedPtr<FCameraShakeData>>
+{
+public:
+	SLATE_BEGIN_ARGS(SCameraShakeRow) {}
+	SLATE_ARGUMENT(TSharedPtr<FCameraShakeData>, CameraShake)
+	SLATE_END_ARGS()
+
+	void Construct(const FArguments& Args, const TSharedRef<STableViewBase>& OwnerTableView)
+	{
+		CameraShake = Args._CameraShake;
+
+		FSuperRowType::Construct(
+			FSuperRowType::FArguments().Padding(1.0f),
+			OwnerTableView);
+	}
+
+	virtual TSharedRef<SWidget> GenerateWidgetForColumn(const FName& ColumnName) override
+	{
+		if (ColumnName == TEXT("CameraShakeName"))
+		{
+			return SNew(SHorizontalBox)
+				+ SHorizontalBox::Slot()
+				.AutoWidth()
+				.Padding(12.0f, 10.0f)
+				.VAlign(VAlign_Center)
+				[
+					SNew(STextBlock)
+					.Text(this, &SCameraShakeRow::GetCameraShakeName)
+				];
+		}
+		else if (ColumnName == TEXT("SceneActorName"))
+		{
+			return SNew(SHorizontalBox)
+				+SHorizontalBox::Slot()
+				.AutoWidth()
+				.Padding(12.0f, 10.0f)
+				.VAlign(VAlign_Center)
+				[
+					SNew(STextBlock)
+					.Text(this, &SCameraShakeRow::GetOwnerActorName)
+				];
+		}
+		else if (ColumnName == TEXT("Status"))
+		{
+			return SNew(SHorizontalBox)
+				+SHorizontalBox::Slot()
+				.AutoWidth()
+				.Padding(12.0f, 10.0f)
+				.VAlign(VAlign_Bottom)
+				[
+					SNew(STextBlock)
+					.Text(this, &SCameraShakeRow::GetCameraShakeStatus)
+				];
+		}
+		return SNullWidget::NullWidget;
+	}
+
+private:
+	FText GetCameraShakeName() const
+	{
+		if (CameraShake->ShakeClass != nullptr)
+		{
+			return FText::FromString(CameraShake->ShakeClass->GetName());
+		}
+		return FText::FromString("<None>");
+	}
+
+	FText GetOwnerActorName() const
+	{
+		if (CameraShake->SourceComponent.IsValid())
+		{
+			if (const AActor* Actor = CameraShake->SourceComponent->GetOwner())
+			{
+				return FText::FromString(Actor->GetName());
+			}
+		}
+		return FText::GetEmpty();
+	}
+
+	FText GetCameraShakeStatus() const
+	{
+		if (CameraShake->bIsHidden)
+		{
+			return FText::FromString("(Hidden)");
+		}
+		return CameraShake->bIsPlaying ? FText::FromString("Playing") : FText::FromString("Stopped");
+	}
+
+	TSharedPtr<FCameraShakeData> CameraShake;
+};
+
+void SCameraShakePreviewer::Construct(const FArguments& InArgs)
+{
+	ChildSlot
+	[
+		SNew(SVerticalBox)
+		+SVerticalBox::Slot()
+		.FillHeight(1.f)
+		[
+			SAssignNew(CameraShakesListView, SListView<TSharedPtr<FCameraShakeData>>)
+			.ListItemsSource(&CameraShakes)
+			.OnGenerateRow(this, &SCameraShakePreviewer::OnCameraShakesListGenerateRowWidget)
+			.OnSelectionChanged(this, &SCameraShakePreviewer::OnCameraShakesListSectionChanged)
+			.HeaderRow
+			(
+				SNew(SHeaderRow)
+				+ SHeaderRow::Column("CameraShakeName")
+				.DefaultLabel(LOCTEXT("CameraShakeName", "Camera Shake Name"))
+				.FillWidth(0.3f)
+				+SHeaderRow::Column("SceneActorName")
+				.DefaultLabel(LOCTEXT("SceneActorName", "Scene Actor Name"))
+				.FillWidth(0.3f)
+				+SHeaderRow::Column("Status")
+				.DefaultLabel(LOCTEXT("Status", "Status"))
+				.FillWidth(0.3f)
+			)
+		]
+		+SVerticalBox::Slot()
+		.AutoHeight()
+		[
+			SNew(SBorder)
+			.BorderImage(FCoreStyle::Get().GetBrush("ToolPanel.GroupBorder"))
+			.Padding(1.0f)
+			[
+				SNew(SWrapBox)
+				.UseAllottedWidth(true)
+				+SWrapBox::Slot()
+				.Padding(FMargin(2.0f))
+				[
+					SNew(SButton)
+					.OnClicked(this, &SCameraShakePreviewer::OnPlayStopAllShakes)
+					.ButtonStyle(FEditorStyle::Get(), "FlatButton.Success")
+					.HAlign(HAlign_Center)
+					.Text(LOCTEXT("PlayStopAllShakes", "Play/Stop All"))
+					.ToolTipText(LOCTEXT("PlayStopAllShakesTooltip", "Play/stop all shakes in the list"))
+				]
+				+SWrapBox::Slot()
+				.Padding(FMargin(2.0f))
+				[
+					SAssignNew(PlayStopSelectedButton, SButton)
+					.OnClicked(this, &SCameraShakePreviewer::OnPlayStopSelectedShake)
+					.IsEnabled(false)
+					.Text(LOCTEXT("PlayStopSelectedShake", "Play/Stop Selected"))
+					.ToolTipText(LOCTEXT("PlayStopSelectedShakeTooltip", "Play/stop select shake"))
+				]
+				+ SWrapBox::Slot()
+				.Padding(FMargin(2.0f))
+				.VAlign(VAlign_Center)
+				[
+					SNew(STextBlock)
+					.Text(LOCTEXT("ActiveViewport", "Active Viewport:"))
+				]
+				+SWrapBox::Slot()
+				.Padding(FMargin(2.0f))
+				.VAlign(VAlign_Center)
+				[
+					SNew(STextBlock)
+					.Text(this, &SCameraShakePreviewer::GetActiveViewportName)
+				]
+				+ SWrapBox::Slot()
+				.Padding(FMargin(2.0f))
+				.VAlign(VAlign_Center)
+				[
+					SNew(STextBlock)
+					.TextStyle(FEditorStyle::Get(), TEXT("Log.Warning"))
+					.Text(this, &SCameraShakePreviewer::GetActiveViewportWarnings)
+				]
+			]
+		]
+	];
+
+	// Listen to when the user toggles support for previewing camera shakes. We'll want to add/remove our view modifiers on
+	// the affected viewport when it happens.
+	CameraShakePreviewerModule = &FModuleManager::GetModuleChecked<FCameraShakePreviewerModule>("CameraShakePreviewer");
+	CameraShakePreviewerModule->OnTogglePreviewCameraShakes.AddSP(this, &SCameraShakePreviewer::OnTogglePreviewCameraShakes);
+
+	// Register callbacks for when stuff happens that might affect the list of shake sources.
+	FEditorDelegates::MapChange.AddSP(this, &SCameraShakePreviewer::OnMapChange);
+	FEditorDelegates::NewCurrentLevel.AddSP(this, &SCameraShakePreviewer::OnNewCurrentLevel);
+	
+	FWorldDelegates::LevelAddedToWorld.AddSP(this, &SCameraShakePreviewer::OnLevelAdded);
+	FWorldDelegates::LevelRemovedFromWorld.AddSP(this, &SCameraShakePreviewer::OnLevelRemoved);
+
+	if (GEngine != nullptr)
+	{
+		GEngine->OnLevelActorListChanged().AddSP(this, &SCameraShakePreviewer::OnLevelActorListChanged);
+		GEngine->OnLevelActorAdded().AddSP(this, &SCameraShakePreviewer::OnLevelActorsAdded);
+		GEngine->OnLevelActorDeleted().AddSP(this, &SCameraShakePreviewer::OnLevelActorsRemoved);
+	}
+	if (GEditor != nullptr)
+	{
+		GEditor->RegisterForUndo(this);
+		GEditor->OnLevelViewportClientListChanged().AddSP(this, &SCameraShakePreviewer::OnLevelViewportClientListChanged);
+	}
+
+	// Create our camera shake manager.
+	PreviewCameraShake = NewObject<UCameraModifier_CameraShake>();
+	ActiveViewportClient = nullptr;
+
+	// Populate the main list based on the current level.
+	Populate();
+	bNeedsRefresh = false;
+}
+
+SCameraShakePreviewer::~SCameraShakePreviewer()
+{
+	PreviewCameraShake = nullptr;
+
+	if (ActiveViewportClient != nullptr)
+	{
+		ActiveViewportClient->ViewModifiers.RemoveAll(this);
+	}
+
+	FEditorDelegates::MapChange.RemoveAll(this);
+	FEditorDelegates::NewCurrentLevel.RemoveAll(this);
+	
+	if (GEditor != nullptr)
+	{
+		GEditor->UnregisterForUndo(this);
+		GEditor->OnLevelViewportClientListChanged().RemoveAll(this);
+	}
+	if (GEngine != nullptr)
+	{
+		GEngine->OnLevelActorListChanged().RemoveAll(this);
+		GEngine->OnLevelActorAdded().RemoveAll(this);
+		GEngine->OnLevelActorDeleted().RemoveAll(this);
+	}
+	
+	FWorldDelegates::LevelAddedToWorld.RemoveAll(this);
+	FWorldDelegates::LevelRemovedFromWorld.RemoveAll(this);
+
+	CameraShakePreviewerModule->OnTogglePreviewCameraShakes.RemoveAll(this);
+}
+
+TSharedRef<ITableRow> SCameraShakePreviewer::OnCameraShakesListGenerateRowWidget(TSharedPtr<FCameraShakeData> CameraShake, const TSharedRef<STableViewBase>& OwnerTable) const
+{
+	return SNew(SCameraShakeRow, OwnerTable)
+		.CameraShake(CameraShake);
+}
+
+bool SCameraShakePreviewer::FindCurrentWorld()
+{
+	CurrentWorld = nullptr;
+	for (const FWorldContext& Context : GEngine->GetWorldContexts())
+	{
+		if (Context.WorldType == EWorldType::Editor)
+		{
+			CurrentWorld = Context.World();
+		}
+	}
+	return CurrentWorld.IsValid();
+}
+
+void SCameraShakePreviewer::Populate()
+{
+	const bool bFoundWorld = FindCurrentWorld();
+	if (!ensureMsgf(bFoundWorld, TEXT("Could not find current world instance.")))
+	{
+		return;
+	}
+
+	// Get all the shake sources from the level.
+	TArray<UCameraShakeSourceComponent*> ShakeSourceComponents;
+	for (ULevel* Level : CurrentWorld->GetLevels())
+	{
+		if (!Level->bIsVisible)
+		{
+			continue;
+		}
+
+		for (AActor* Actor : Level->Actors)
+		{
+			ACameraShakeSourceActor* ShakeSourceActor = Cast<ACameraShakeSourceActor>(Actor);
+			if (ShakeSourceActor != nullptr)
+			{
+				if (UCameraShakeSourceComponent* ShakeSourceComponent = ShakeSourceActor->GetCameraShakeSourceComponent())
+				{
+					ShakeSourceComponents.Add(ShakeSourceComponent);
+				}
+			}
+		}
+	}
+	TSet<UCameraShakeSourceComponent*> ShakeSourceComponentsSet(ShakeSourceComponents);
+
+	// Get all the shake sources we already knew about.
+	TSet<TSharedPtr<FCameraShakeData>> RemovedShakes;
+	TSet<UCameraShakeSourceComponent*> PreviousShakeSourceComponentsSet;
+	for (TSharedPtr<FCameraShakeData> CameraShake : CameraShakes)
+	{
+		if (CameraShake->SourceComponent.IsValid())
+		{
+			UCameraShakeSourceComponent* ShakeSourceComponent = CameraShake->SourceComponent.Get();
+			if (ShakeSourceComponentsSet.Contains(ShakeSourceComponent))
+			{
+				PreviousShakeSourceComponentsSet.Add(ShakeSourceComponent);
+			}
+			else
+			{
+				RemovedShakes.Add(CameraShake);
+			}
+		}
+		else if (CameraShake->SourceComponent.IsStale())
+		{
+			RemovedShakes.Add(CameraShake);
+		}
+	}
+
+	// Remove the shakes that have sources that were destroyed. We don't need to stop them if they were running, 
+	// the `CameraModifier_CameraShake` class will clean those up automatically.
+	for (TSharedPtr<FCameraShakeData> RemovedShake : RemovedShakes)
+	{
+		if (RemovedShake->SourceComponent.IsValid())
+		{
+			PreviewCameraShake->RemoveAllCameraShakesFromSource(RemovedShake->SourceComponent.Get());
+		}
+		CameraShakes.Remove(RemovedShake);
+	}
+
+	// Add new camera shake sources to the list.
+	TSet<UCameraShakeSourceComponent*> NewShakeSourceComponents = ShakeSourceComponentsSet.Difference(PreviousShakeSourceComponentsSet);
+	for (UCameraShakeSourceComponent* ShakeSourceComponent : NewShakeSourceComponents)
+	{
+		TSharedPtr<FCameraShakeData> CameraShake = MakeShared<FCameraShakeData>();
+		CameraShake->ShakeClass = ShakeSourceComponent->CameraShake;
+		CameraShake->SourceComponent = ShakeSourceComponent;
+		CameraShakes.Add(CameraShake);
+	}
+
+	CameraShakesListView->RequestListRefresh();
+}
+
+void SCameraShakePreviewer::Refresh()
+{
+	bNeedsRefresh = true;
+}
+
+void SCameraShakePreviewer::Tick(const FGeometry& AllottedGeometry, const double InCurrentTime, const float InDeltaTime)
+{
+	SCompoundWidget::Tick(AllottedGeometry, InCurrentTime, InDeltaTime);
+
+	// Update our list of camera shakes if needed.
+	if (bNeedsRefresh)
+	{
+		Populate();
+		bNeedsRefresh = false;
+	}
+
+	// Update playing information on our camera shakes.
+	for (TSharedPtr<FCameraShakeData> CameraShake : CameraShakes)
+	{
+		if (CameraShake->SourceComponent.IsValid())
+		{
+			UCameraShakeSourceComponent* ShakeSourceComponent = CameraShake->SourceComponent.Get();
+
+			// Handle the case where the actor was hidden: stop it if it was playing, flag it as hidden.
+			if (AActor* ShakeSourceActor = ShakeSourceComponent->GetOwner())
+			{
+				const bool bIsHidden = ShakeSourceActor->IsHiddenEd();
+				CameraShake->bIsHidden = bIsHidden;
+				if (bIsHidden && CameraShake->bIsPlaying)
+				{
+					PreviewCameraShake->RemoveAllCameraShakesFromSource(ShakeSourceComponent);
+					CameraShake->bIsPlaying = false;
+				}
+			}
+
+			// Handle the case where the user has changed the shake type on an existing shake source.
+			if (ShakeSourceComponent->CameraShake != CameraShake->ShakeClass)
+			{
+				// Yep, it's a different class now. Let's stop any running instance of the old class,
+				// and figure out if we should start playing the new class right away (if the old class instance
+				// was running until now).
+				CameraShake->ShakeClass = ShakeSourceComponent->CameraShake;
+				CameraShake->ShakeInstance = nullptr;
+
+				PreviewCameraShake->RemoveAllCameraShakesFromSource(ShakeSourceComponent);
+
+				if (CameraShake->bIsPlaying)
+				{
+					if (ShakeSourceComponent->CameraShake != nullptr)
+					{
+						PlayCameraShake(CameraShake);
+					}
+					else
+					{
+						CameraShake->bIsPlaying = false;
+					}
+				}
+			}
+		}
+	}
+
+	// See if the active viewport has changed.
+	if (GEditor)
+	{
+		FViewport* ActiveViewport = GEditor->GetActiveViewport();
+		if ((ActiveViewportClient == nullptr && ActiveViewport != nullptr) ||
+			(ActiveViewportClient != nullptr && ActiveViewportClient->Viewport != ActiveViewport))
+		{
+			ActiveViewportIndex = 0;
+			FLevelEditorViewportClient* NewActiveViewportClient = nullptr;
+			if (ActiveViewport != nullptr)
+			{
+				// Find the new active viewport client.
+				const TArray<FLevelEditorViewportClient*>& ViewportClients = GEditor->GetLevelViewportClients();
+				for (FLevelEditorViewportClient* ViewportClient : ViewportClients)
+				{
+					++ActiveViewportIndex;
+					if (ViewportClient->Viewport == ActiveViewport)
+					{
+						NewActiveViewportClient = ViewportClient;
+						break;
+					}
+				}
+			}
+			// else: no new active viewport client, and leave the index at 0.
+
+			if (NewActiveViewportClient != ActiveViewportClient)
+			{
+				// Clear the old viewport's callbacks.
+				if (ActiveViewportClient != nullptr)
+				{
+					ActiveViewportClient->ViewModifiers.RemoveAll(this);
+				}
+
+				ActiveViewportClient = NewActiveViewportClient;
+
+				// Add to the new viewport's callbacks if necessary.
+				if (ActiveViewportClient != nullptr)
+				{
+					if (ActiveViewportClient != nullptr && CameraShakePreviewerModule->HasCameraShakesPreview(ActiveViewportClient))
+					{
+						ActiveViewportClient->ViewModifiers.AddRaw(this, &SCameraShakePreviewer::OnModifyView);
+					}
+				}
+			}
+		}
+	}
+}
+
+void SCameraShakePreviewer::OnTogglePreviewCameraShakes(const FTogglePreviewCameraShakesParams& Params)
+{
+	if (Params.ViewportClient == ActiveViewportClient)
+	{
+		if (Params.bPreviewCameraShakes && ensure(!Params.ViewportClient->ViewModifiers.IsBoundToObject(this)))
+		{
+			ActiveViewportClient->ViewModifiers.AddRaw(this, &SCameraShakePreviewer::OnModifyView);
+		}
+		else if (!Params.bPreviewCameraShakes && ensure(Params.ViewportClient->ViewModifiers.IsBoundToObject(this)))
+		{
+			ActiveViewportClient->ViewModifiers.RemoveAll(this);
+		}
+	}
+}
+
+void SCameraShakePreviewer::OnCameraShakesListSectionChanged(TSharedPtr<FCameraShakeData> Entry, ESelectInfo::Type SelectInfo) const
+{
+	PlayStopSelectedButton->SetEnabled(Entry != nullptr);
+
+	GEditor->SelectNone(true, true, false);
+	if (Entry != nullptr && Entry->SourceComponent.IsValid())
+	{
+		AActor* SourceActor = Entry->SourceComponent.Get()->GetOwner();
+		GEditor->SelectActor(SourceActor, true, true);
+	}
+}
+
+FText SCameraShakePreviewer::GetActiveViewportName() const
+{
+	if (ActiveViewportIndex > 0)
+	{
+		return FText::FromString(LexToString(ActiveViewportIndex));
+	}
+	return FText::FromString("<None>");
+}
+
+FText SCameraShakePreviewer::GetActiveViewportWarnings() const
+{
+	FText Warnings;
+	if (ActiveViewportClient != nullptr)
+	{
+		FText Delimiter = FText::FromString(", ");
+		if (!CameraShakePreviewerModule->HasCameraShakesPreview(ActiveViewportClient))
+		{
+			// Ya can't see no shakes if yer not previewing no shakes.
+			FText Warning = LOCTEXT("ActiveViewportPreviewShakesOffWarning", "Camera shakes previewing is off");
+			Warnings = Warnings.IsEmpty() ? Warning : FText::Join(Delimiter, Warnings, Warning);
+		}
+		if (!ActiveViewportClient->IsRealtime())
+		{
+			// When real-time mode is off in a viewport, the viewport is ticked only when the user is actively
+			// doing something, so the shaking comes and goes in a weird way. Might as well warn the user about it.
+			FText Warning = LOCTEXT("ActiveViewportRealtimeOffWarning", "Real-time mode is off");
+			Warnings = Warnings.IsEmpty() ? Warning : FText::Join(Delimiter, Warnings, Warning);
+		}
+	}
+	else
+	{
+		Warnings = LOCTEXT("NoActiveViewportWarning", "No active viewport");
+	}
+	return Warnings;
+}
+
+FReply SCameraShakePreviewer::OnPlayStopAllShakes()
+{
+	TArray<FActiveCameraShakeInfo> ActiveCameraShakes;
+	PreviewCameraShake->GetActiveCameraShakes(ActiveCameraShakes);
+	if (ActiveCameraShakes.Num() > 0)
+	{
+		// If we have at least 1 shake still playing, stop it.
+		PreviewCameraShake->RemoveAllCameraShakes();
+
+		for (TSharedPtr<FCameraShakeData> CameraShake : CameraShakes)
+		{
+			CameraShake->ShakeInstance = nullptr;
+			CameraShake->bIsPlaying = false;
+		}
+	}
+	else
+	{
+		// No shakes playing... start them all.
+		for (TSharedPtr<FCameraShakeData> CameraShake : CameraShakes)
+		{
+			if (CameraShake->ShakeClass != nullptr)
+			{
+				PlayCameraShake(CameraShake);
+			}
+		}
+	}
+	
+	return FReply::Handled();
+}
+
+FReply SCameraShakePreviewer::OnPlayStopSelectedShake()
+{
+	TArray<TSharedPtr<FCameraShakeData>> SelectedItems = CameraShakesListView->GetSelectedItems();
+	for (TSharedPtr<FCameraShakeData> SelectedItem : SelectedItems)
+	{
+		if (!SelectedItem->bIsPlaying && SelectedItem->ShakeClass != nullptr)
+		{
+			PlayCameraShake(SelectedItem);
+		}
+		else if (SelectedItem->bIsPlaying && SelectedItem->ShakeInstance != nullptr)
+		{
+			PreviewCameraShake->RemoveCameraShake(SelectedItem->ShakeInstance);
+			SelectedItem->ShakeInstance = nullptr;
+			SelectedItem->bIsPlaying = false;
+		}
+	}
+	return FReply::Handled();
+}
+
+void SCameraShakePreviewer::PlayCameraShake(TSharedPtr<FCameraShakeData> CameraShake)
+{
+	FAddCameraShakeParams Params;
+	Params.SourceComponent = CameraShake->SourceComponent.Get();
+	UCameraShake* ShakeInstance = PreviewCameraShake->AddCameraShake(CameraShake->ShakeClass, Params);
+	CameraShake->ShakeInstance = ShakeInstance;
+	CameraShake->bIsPlaying = true;
+}
+
+void SCameraShakePreviewer::OnLevelViewportClientListChanged()
+{
+	if (ActiveViewportClient != nullptr)
+	{
+		const TArray<FLevelEditorViewportClient*> ViewportClients = GEditor->GetLevelViewportClients();
+		const bool bFound = ViewportClients.Contains(ActiveViewportClient);
+		if (!bFound)
+		{
+			// Our target viewport has been removed from the list... let's forget it.
+			// TODO: unsafe? maybe we don't actually need to cleanup?
+			ActiveViewportClient->ViewModifiers.RemoveAll(this);
+			ActiveViewportClient = nullptr;
+		}
+	}
+}
+
+void SCameraShakePreviewer::OnLevelAdded(ULevel* InLevel, UWorld* InWorld)
+{
+	Refresh();
+}
+
+void SCameraShakePreviewer::OnLevelRemoved(ULevel* InLevel, UWorld* InWorld)
+{
+	Refresh();
+}
+
+void SCameraShakePreviewer::OnLevelActorsAdded(AActor* InActor)
+{
+	Refresh();
+}
+
+void SCameraShakePreviewer::OnLevelActorsRemoved(AActor* InActor)
+{
+	Refresh();
+}
+
+void SCameraShakePreviewer::OnLevelActorListChanged()
+{
+	Refresh();
+}
+
+void SCameraShakePreviewer::OnMapChange(uint32 MapFlags)
+{
+	Refresh();
+}
+
+void SCameraShakePreviewer::OnNewCurrentLevel()
+{
+	Refresh();
+}
+
+void SCameraShakePreviewer::OnMapLoaded(const FString&  Filename, bool bAsTemplate)
+{
+	Refresh();
+}
+
+void SCameraShakePreviewer::PostUndo(bool bSuccess)
+{
+	Refresh();
+}
+
+void SCameraShakePreviewer::OnModifyView(const FEditorViewportViewModifierParams& Params, FMinimalViewInfo& InOutPOV)
+{
+	PreviewCameraShake->ModifyCamera(Params.DeltaTime, InOutPOV);
+}
+
+void SCameraShakePreviewer::AddReferencedObjects(FReferenceCollector& Collector)
+{
+	Collector.AddReferencedObject(PreviewCameraShake);
+}
+
+FString SCameraShakePreviewer::GetReferencerName() const
+{
+	return TEXT("SCameraShakePreviewer");
+}
+
+#undef LOCTEXT_NAMESPACE
diff --git a/Engine/Plugins/Cameras/CameraShakePreviewer/Source/CameraShakePreviewer/Private/SCameraShakePreviewer.h b/Engine/Plugins/Cameras/CameraShakePreviewer/Source/CameraShakePreviewer/Private/SCameraShakePreviewer.h
new file mode 100644
--- /dev/null
+++ b/Engine/Plugins/Cameras/CameraShakePreviewer/Source/CameraShakePreviewer/Private/SCameraShakePreviewer.h
@@ -0,0 +1,92 @@
+// Copyright 1998-2019 Epic Games, Inc. All Rights Reserved.
+
+#pragma once 
+
+#include "Containers/Array.h"
+#include "EditorUndoClient.h"
+#include "Templates/SharedPointer.h"
+#include "UObject/GCObject.h"
+#include "Widgets/DeclarativeSyntaxSupport.h"
+#include "Widgets/Views/SListView.h"
+#include "Widgets/SCompoundWidget.h"
+
+class AActor;
+class FCameraShakePreviewerLevelEditorViewportClient;
+class FLevelEditorViewportClient;
+class FSceneView;
+class ITableRow;
+class STableViewBase;
+class ULevel;
+class UCameraModifier_CameraShake;
+class UCameraShakeSourceComponent;
+class UWorld;
+struct FCameraShakeData;
+struct FEditorViewportViewModifierParams;
+struct FMinimalViewInfo;
+struct FTogglePreviewCameraShakesParams;
+
+/**
+ * Camera shake preview panel.
+ */
+class SCameraShakePreviewer : public SCompoundWidget, public FEditorUndoClient, public FGCObject
+{
+public:
+	SLATE_BEGIN_ARGS(SCameraShakePreviewer) {}
+	SLATE_END_ARGS()
+
+	void Construct(const FArguments& InArgs);
+	~SCameraShakePreviewer();
+
+	// SWidget interface
+	virtual void Tick(const FGeometry& AllottedGeometry, const double InCurrentTime, const float InDeltaTime) override;
+
+	// FEditorUndoClient Interface
+	virtual void PostUndo(bool bSuccess) override;
+	virtual void PostRedo(bool bSuccess) override { PostUndo(bSuccess); }
+
+	// FGCObject interface
+	virtual void AddReferencedObjects(FReferenceCollector& Collector) override;
+	virtual FString GetReferencerName() const override;
+
+private:
+	void Populate();
+	bool FindCurrentWorld();
+	void Refresh();
+
+	void OnTogglePreviewCameraShakes(const FTogglePreviewCameraShakesParams& Params);
+
+	TSharedRef<ITableRow> OnCameraShakesListGenerateRowWidget(TSharedPtr<FCameraShakeData> Entry, const TSharedRef<STableViewBase>& OwnerTable) const;
+	void OnCameraShakesListSectionChanged(TSharedPtr<FCameraShakeData> Entry, ESelectInfo::Type SelectInfo) const;
+	FText GetActiveViewportName() const;
+	FText GetActiveViewportWarnings() const;
+
+	FReply OnPlayStopAllShakes();
+	FReply OnPlayStopSelectedShake();
+	void PlayCameraShake(TSharedPtr<FCameraShakeData> CameraShake);
+
+	void OnLevelViewportClientListChanged();
+	void OnLevelAdded(ULevel* InLevel, UWorld* InWorld);
+	void OnLevelRemoved(ULevel* InLevel, UWorld* InWorld);
+	void OnLevelActorsAdded(AActor* InActor);
+	void OnLevelActorsRemoved(AActor* InActor);
+	void OnLevelActorListChanged();
+	void OnMapChange(uint32 MapFlags);
+	void OnNewCurrentLevel();
+	void OnMapLoaded(const FString&  Filename, bool bAsTemplate);
+
+	void OnModifyView(const FEditorViewportViewModifierParams& Params, FMinimalViewInfo& InOutPOV);
+
+private:
+	TArray<TSharedPtr<FCameraShakeData>> CameraShakes;
+	UCameraModifier_CameraShake* PreviewCameraShake;
+
+	TSharedPtr<SListView<TSharedPtr<FCameraShakeData>>> CameraShakesListView;
+	TSharedPtr<SButton> PlayStopSelectedButton;
+
+	FCameraShakePreviewerModule* CameraShakePreviewerModule;
+	FLevelEditorViewportClient* ActiveViewportClient;
+	int ActiveViewportIndex;
+
+	TWeakObjectPtr<UWorld> CurrentWorld;
+	bool bNeedsRefresh;
+};
\ No newline at end of file
diff --git a/Engine/Plugins/Cameras/CameraShakePreviewer/Source/CameraShakePreviewer/Public/CameraShakePreviewerModule.h b/Engine/Plugins/Cameras/CameraShakePreviewer/Source/CameraShakePreviewer/Public/CameraShakePreviewerModule.h
new file mode 100644
--- /dev/null
+++ b/Engine/Plugins/Cameras/CameraShakePreviewer/Source/CameraShakePreviewer/Public/CameraShakePreviewerModule.h
@@ -0,0 +1,62 @@
+// Copyright 1998-2019 Epic Games, Inc. All Rights Reserved.
+#pragma once
+
+#include "CoreMinimal.h"
+#include "Modules/ModuleInterface.h"
+
+class FExtender;
+class FUICommandList;
+class FLevelEditorViewportClient;
+class SDockTab;
+class SWidget;
+
+struct FTogglePreviewCameraShakesParams
+{
+	FLevelEditorViewportClient* ViewportClient = nullptr;
+	bool bPreviewCameraShakes = false;
+};
+
+DECLARE_MULTICAST_DELEGATE_OneParam(FOnTogglePreviewCameraShakes, const FTogglePreviewCameraShakesParams&);
+
+/**
+ * Module for managing the camera shake previewer tool.
+ */
+class FCameraShakePreviewerModule : public IModuleInterface
+{
+public:
+	// IModuleInterface interface
+	virtual void StartupModule();
+	virtual void ShutdownModule();
+
+	// Toggles whether the given viewport should show camera shakes.
+	void ToggleCameraShakesPreview(FLevelEditorViewportClient* ViewportClient);
+	// Gets whether the given viewport is supporting camera shakes.
+	bool HasCameraShakesPreview(FLevelEditorViewportClient* ViewportClient) const;
+
+public:
+	// Callback for when a viewport's ability to show camera shakes is toggled.
+	FOnTogglePreviewCameraShakes OnTogglePreviewCameraShakes;
+	
+private:
+	static TSharedRef<SDockTab> CreateCameraShakePreviewerTab(const FSpawnTabArgs& Args);
+
+	void RegisterEditorTab();
+	void RegisterViewportOptionMenuExtender();
+
+	void UnregisterEditorTab();
+	void UnregisterViewportOptionMenuExtender();
+
+	TSharedRef<FExtender> OnExtendLevelViewportOptionMenu(const TSharedRef<FUICommandList> CommandList);
+	void OnLevelViewportClientListChanged();
+
+private:
+	FDelegateHandle LevelEditorTabManagerChangedHandle;
+	FDelegateHandle ViewportOptionsMenuExtenderHandle;
+
+	struct FViewportInfo
+	{
+		bool bPreviewCameraShakes = false;
+	};
+	TMap<FLevelEditorViewportClient*, FViewportInfo> ViewportInfos;
+};
+
