Fix non-unity build

diff --git a/Engine/Plugins/Cameras/CameraShakePreviewer/Source/CameraShakePreviewer/Public/CameraShakePreviewerModule.h b/Engine/Plugins/Cameras/CameraShakePreviewer/Source/CameraShakePreviewer/Public/CameraShakePreviewerModule.h
--- a/Engine/Plugins/Cameras/CameraShakePreviewer/Source/CameraShakePreviewer/Public/CameraShakePreviewerModule.h
+++ b/Engine/Plugins/Cameras/CameraShakePreviewer/Source/CameraShakePreviewer/Public/CameraShakePreviewerModule.h
@@ -7,6 +7,7 @@
 class FExtender;
 class FUICommandList;
 class FLevelEditorViewportClient;
+class FSpawnTabArgs;
 class SDockTab;
 class SWidget;
 
