Fix non-unity build.

diff --git a/Engine/Plugins/Cameras/CameraShakePreviewer/Source/CameraShakePreviewer/Private/SCameraShakePreviewer.h b/Engine/Plugins/Cameras/CameraShakePreviewer/Source/CameraShakePreviewer/Private/SCameraShakePreviewer.h
--- a/Engine/Plugins/Cameras/CameraShakePreviewer/Source/CameraShakePreviewer/Private/SCameraShakePreviewer.h
+++ b/Engine/Plugins/Cameras/CameraShakePreviewer/Source/CameraShakePreviewer/Private/SCameraShakePreviewer.h
@@ -12,6 +12,7 @@
 
 class AActor;
 class FCameraShakePreviewerLevelEditorViewportClient;
+class FCameraShakePreviewerModule;
 class FLevelEditorViewportClient;
 class FSceneView;
 class ITableRow;
