Trying to fix main build.

diff --git a/Engine/Plugins/Cameras/CameraShakePreviewer/Source/CameraShakePreviewer/Private/CameraShakePreviewerModule.cpp b/Engine/Plugins/Cameras/CameraShakePreviewer/Source/CameraShakePreviewer/Private/CameraShakePreviewerModule.cpp
--- a/Engine/Plugins/Cameras/CameraShakePreviewer/Source/CameraShakePreviewer/Private/CameraShakePreviewerModule.cpp
+++ b/Engine/Plugins/Cameras/CameraShakePreviewer/Source/CameraShakePreviewer/Private/CameraShakePreviewerModule.cpp
@@ -16,6 +16,8 @@
 #include "Widgets/SWidget.h"
 #include "WorkspaceMenuStructure.h"
 #include "WorkspaceMenuStructureModule.h"
+#include "IAssetViewport.h"
+#include "EditorViewportClient.h"
 
 #define LOCTEXT_NAMESPACE "CameraShakePreviewer"
 
@@ -133,8 +135,8 @@ TSharedRef<FExtender> FCameraShakePreviewerModule::OnExtendLevelViewportOptionMe
 	// Find the viewport for which we're opening this options menu.
 	FLevelEditorModule& LevelEditorModule = FModuleManager::LoadModuleChecked<FLevelEditorModule>(LevelEditorModuleName);
 	TSharedPtr<ILevelEditor> LevelEditor = LevelEditorModule.GetLevelEditorInstance().Pin();
-	TSharedPtr<ILevelViewport> ViewportInterface = LevelEditor->GetActiveViewportInterface();
-	FLevelEditorViewportClient* ViewportClient = &ViewportInterface->GetLevelViewportClient();
+	TSharedPtr<IAssetViewport> ViewportInterface = LevelEditor->GetActiveViewportInterface();
+	FLevelEditorViewportClient* ViewportClient = (FLevelEditorViewportClient*)&ViewportInterface->GetAssetViewportClient();
 	if (ViewportClient->ViewportType != ELevelViewportType::LVT_Perspective)
 	{
 		// If the current active viewport isn't perspective, let's not add any option for previewing camera shakes.
