Editor: Make view-modifiers stateless, only apply them in `CalcSceneView`.

This puts the burden of ticking and state on the view-modifiers themselves, so this change also updates the camera shake preview tool to do that on its own.

diff --git a/Engine/Plugins/Cameras/CameraShakePreviewer/Source/CameraShakePreviewer/Private/SCameraShakePreviewer.cpp b/Engine/Plugins/Cameras/CameraShakePreviewer/Source/CameraShakePreviewer/Private/SCameraShakePreviewer.cpp
--- a/Engine/Plugins/Cameras/CameraShakePreviewer/Source/CameraShakePreviewer/Private/SCameraShakePreviewer.cpp
+++ b/Engine/Plugins/Cameras/CameraShakePreviewer/Source/CameraShakePreviewer/Private/SCameraShakePreviewer.cpp
@@ -10,6 +10,7 @@
 #include "LevelEditorViewport.h"
 #include "Modules/ModuleManager.h"
 #include "Slate/SceneViewport.h"
+#include "TickableEditorObject.h"
 #include "Widgets/Input/SButton.h"
 #include "Widgets/Input/SCheckBox.h"
 #include "Widgets/Layout/SBox.h"
@@ -40,6 +41,35 @@ struct FCameraShakeData
 	}
 };
 
+FCameraShakePreviewUpdater::FCameraShakePreviewUpdater()
+	: PreviewCameraShake(NewObject<UCameraModifier_CameraShake>())
+	, LastLocationModifier(FVector::ZeroVector)
+	, LastRotationModifier(FRotator::ZeroRotator)
+	, LastFOVModifier(0.f)
+{}
+
+void FCameraShakePreviewUpdater::ModifyCamera(FMinimalViewInfo& InOutPOV)
+{
+	const float DeltaTime = LastDeltaTime.Get(-1.f);
+	if (DeltaTime > 0.f)
+	{
+		FMinimalViewInfo InPOV(InOutPOV);
+		PreviewCameraShake->ModifyCamera(DeltaTime, InOutPOV);
+
+		LastLocationModifier = InOutPOV.Location - InPOV.Location;
+		LastRotationModifier = InOutPOV.Rotation - InPOV.Rotation;
+		LastFOVModifier = InOutPOV.FOV - InPOV.FOV;
+
+		LastDeltaTime.Reset();
+	}
+	else
+	{
+		InOutPOV.Location += LastLocationModifier;
+		InOutPOV.Rotation += LastRotationModifier;
+		InOutPOV.FOV += LastFOVModifier;
+	}
+}
+
 /**
  * The UI for each entry in the panel's main list.
  */
@@ -239,7 +269,7 @@ void SCameraShakePreviewer::Construct(const FArguments& InArgs)
 	}
 
 	// Create our camera shake manager.
-	PreviewCameraShake = NewObject<UCameraModifier_CameraShake>();
+	CameraShakePreviewUpdater = TUniquePtr<FCameraShakePreviewUpdater>(new FCameraShakePreviewUpdater());
 	ActiveViewportClient = nullptr;
 
 	// Populate the main list based on the current level.
@@ -249,7 +279,7 @@ void SCameraShakePreviewer::Construct(const FArguments& InArgs)
 
 SCameraShakePreviewer::~SCameraShakePreviewer()
 {
-	PreviewCameraShake = nullptr;
+	CameraShakePreviewUpdater = nullptr;
 
 	if (ActiveViewportClient != nullptr)
 	{
@@ -356,7 +386,7 @@ void SCameraShakePreviewer::Populate()
 	{
 		if (RemovedShake->SourceComponent.IsValid())
 		{
-			PreviewCameraShake->RemoveAllCameraShakesFromSource(RemovedShake->SourceComponent.Get());
+			CameraShakePreviewUpdater->ShakeModifier().RemoveAllCameraShakesFromSource(RemovedShake->SourceComponent.Get());
 		}
 		CameraShakes.Remove(RemovedShake);
 	}
@@ -404,7 +434,7 @@ void SCameraShakePreviewer::Tick(const FGeometry& AllottedGeometry, const double
 				CameraShake->bIsHidden = bIsHidden;
 				if (bIsHidden && CameraShake->bIsPlaying)
 				{
-					PreviewCameraShake->RemoveAllCameraShakesFromSource(ShakeSourceComponent);
+					CameraShakePreviewUpdater->ShakeModifier().RemoveAllCameraShakesFromSource(ShakeSourceComponent);
 					CameraShake->bIsPlaying = false;
 				}
 			}
@@ -418,7 +448,7 @@ void SCameraShakePreviewer::Tick(const FGeometry& AllottedGeometry, const double
 				CameraShake->ShakeClass = ShakeSourceComponent->CameraShake;
 				CameraShake->ShakeInstance = nullptr;
 
-				PreviewCameraShake->RemoveAllCameraShakesFromSource(ShakeSourceComponent);
+				CameraShakePreviewUpdater->ShakeModifier().RemoveAllCameraShakesFromSource(ShakeSourceComponent);
 
 				if (CameraShake->bIsPlaying)
 				{
@@ -549,11 +579,11 @@ FText SCameraShakePreviewer::GetActiveViewportWarnings() const
 FReply SCameraShakePreviewer::OnPlayStopAllShakes()
 {
 	TArray<FActiveCameraShakeInfo> ActiveCameraShakes;
-	PreviewCameraShake->GetActiveCameraShakes(ActiveCameraShakes);
+	CameraShakePreviewUpdater->ShakeModifier().GetActiveCameraShakes(ActiveCameraShakes);
 	if (ActiveCameraShakes.Num() > 0)
 	{
 		// If we have at least 1 shake still playing, stop it.
-		PreviewCameraShake->RemoveAllCameraShakes();
+		CameraShakePreviewUpdater->ShakeModifier().RemoveAllCameraShakes();
 
 		for (TSharedPtr<FCameraShakeData> CameraShake : CameraShakes)
 		{
@@ -587,7 +617,7 @@ FReply SCameraShakePreviewer::OnPlayStopSelectedShake()
 		}
 		else if (SelectedItem->bIsPlaying && SelectedItem->ShakeInstance != nullptr)
 		{
-			PreviewCameraShake->RemoveCameraShake(SelectedItem->ShakeInstance);
+			CameraShakePreviewUpdater->ShakeModifier().RemoveCameraShake(SelectedItem->ShakeInstance);
 			SelectedItem->ShakeInstance = nullptr;
 			SelectedItem->bIsPlaying = false;
 		}
@@ -599,7 +629,7 @@ void SCameraShakePreviewer::PlayCameraShake(TSharedPtr<FCameraShakeData> CameraS
 {
 	FAddCameraShakeParams Params;
 	Params.SourceComponent = CameraShake->SourceComponent.Get();
-	UCameraShake* ShakeInstance = PreviewCameraShake->AddCameraShake(CameraShake->ShakeClass, Params);
+	UCameraShake* ShakeInstance = CameraShakePreviewUpdater->ShakeModifier().AddCameraShake(CameraShake->ShakeClass, Params);
 	CameraShake->ShakeInstance = ShakeInstance;
 	CameraShake->bIsPlaying = true;
 }
@@ -665,19 +695,9 @@ void SCameraShakePreviewer::PostUndo(bool bSuccess)
 	Refresh();
 }
 
-void SCameraShakePreviewer::OnModifyView(const FEditorViewportViewModifierParams& Params, FMinimalViewInfo& InOutPOV)
-{
-	PreviewCameraShake->ModifyCamera(Params.DeltaTime, InOutPOV);
-}
-
-void SCameraShakePreviewer::AddReferencedObjects(FReferenceCollector& Collector)
-{
-	Collector.AddReferencedObject(PreviewCameraShake);
-}
-
-FString SCameraShakePreviewer::GetReferencerName() const
+void SCameraShakePreviewer::OnModifyView(FMinimalViewInfo& InOutPOV)
 {
-	return TEXT("SCameraShakePreviewer");
+	CameraShakePreviewUpdater->ModifyCamera(InOutPOV);
 }
 
 #undef LOCTEXT_NAMESPACE
diff --git a/Engine/Plugins/Cameras/CameraShakePreviewer/Source/CameraShakePreviewer/Private/SCameraShakePreviewer.h b/Engine/Plugins/Cameras/CameraShakePreviewer/Source/CameraShakePreviewer/Private/SCameraShakePreviewer.h
--- a/Engine/Plugins/Cameras/CameraShakePreviewer/Source/CameraShakePreviewer/Private/SCameraShakePreviewer.h
+++ b/Engine/Plugins/Cameras/CameraShakePreviewer/Source/CameraShakePreviewer/Private/SCameraShakePreviewer.h
@@ -13,6 +13,7 @@
 class AActor;
 class FCameraShakePreviewerLevelEditorViewportClient;
 class FCameraShakePreviewerModule;
+class FCameraShakePreviewUpdater;
 class FLevelEditorViewportClient;
 class FSceneView;
 class ITableRow;
@@ -22,14 +23,41 @@ class UCameraModifier_CameraShake;
 class UCameraShakeSourceComponent;
 class UWorld;
 struct FCameraShakeData;
-struct FEditorViewportViewModifierParams;
 struct FMinimalViewInfo;
 struct FTogglePreviewCameraShakesParams;
 
+class FCameraShakePreviewUpdater : public FTickableEditorObject, public FGCObject
+{
+public:
+	FCameraShakePreviewUpdater();
+
+	// FTickableObject Interface
+	virtual ETickableTickType GetTickableTickType() const { return ETickableTickType::Always; }
+	virtual TStatId GetStatId() const override { RETURN_QUICK_DECLARE_CYCLE_STAT(FCameraShakePreviewUpdater, STATGROUP_Tickables); }
+	virtual void Tick(float DeltaTime) override { LastDeltaTime = DeltaTime; }
+
+	// FGCObject interface
+	virtual void AddReferencedObjects(FReferenceCollector& Collector) override { Collector.AddReferencedObject(PreviewCameraShake); }
+	virtual FString GetReferencerName() const override { return TEXT("SCameraShakePreviewer"); }
+
+	UCameraModifier_CameraShake& ShakeModifier() const { return *PreviewCameraShake; }
+
+	void ModifyCamera(FMinimalViewInfo& InOutPOV);
+
+private:
+	UCameraModifier_CameraShake* PreviewCameraShake;
+
+	TOptional<float> LastDeltaTime;
+
+	FVector LastLocationModifier;
+	FRotator LastRotationModifier;
+	float LastFOVModifier;
+};
+
 /**
  * Camera shake preview panel.
  */
-class SCameraShakePreviewer : public SCompoundWidget, public FEditorUndoClient, public FGCObject
+class SCameraShakePreviewer : public SCompoundWidget, public FEditorUndoClient
 {
 public:
 	SLATE_BEGIN_ARGS(SCameraShakePreviewer) {}
@@ -45,10 +73,6 @@ public:
 	virtual void PostUndo(bool bSuccess) override;
 	virtual void PostRedo(bool bSuccess) override { PostUndo(bSuccess); }
 
-	// FGCObject interface
-	virtual void AddReferencedObjects(FReferenceCollector& Collector) override;
-	virtual FString GetReferencerName() const override;
-
 private:
 	void Populate();
 	bool FindCurrentWorld();
@@ -75,11 +99,11 @@ private:
 	void OnNewCurrentLevel();
 	void OnMapLoaded(const FString&  Filename, bool bAsTemplate);
 
-	void OnModifyView(const FEditorViewportViewModifierParams& Params, FMinimalViewInfo& InOutPOV);
+	void OnModifyView(FMinimalViewInfo& InOutPOV);
 
 private:
 	TArray<TSharedPtr<FCameraShakeData>> CameraShakes;
-	UCameraModifier_CameraShake* PreviewCameraShake;
+	TUniquePtr<FCameraShakePreviewUpdater> CameraShakePreviewUpdater;
 
 	TSharedPtr<SListView<TSharedPtr<FCameraShakeData>>> CameraShakesListView;
 	TSharedPtr<SButton> PlayStopSelectedButton;
diff --git a/Engine/Source/Editor/UnrealEd/Private/EditorViewportClient.cpp b/Engine/Source/Editor/UnrealEd/Private/EditorViewportClient.cpp
--- a/Engine/Source/Editor/UnrealEd/Private/EditorViewportClient.cpp
+++ b/Engine/Source/Editor/UnrealEd/Private/EditorViewportClient.cpp
@@ -331,6 +331,7 @@ FEditorViewportClient::FEditorViewportClient(FEditorModeTools* InModeTools, FPre
 	, bSetListenerPosition(false)
 	, LandscapeLODOverride(-1)
 	, bDrawVertices(false)
+	, bShouldApplyViewModifiers(true)
 	, bOwnsModeTools(false)
 	, ModeTools(InModeTools)
 	, Widget(new FWidget)
@@ -742,9 +743,23 @@ FSceneView* FEditorViewportClient::CalcSceneView(FSceneViewFamily* ViewFamily, c
 	FViewportCameraTransform& ViewTransform = GetViewTransform();
 	const ELevelViewportType EffectiveViewportType = GetViewportType();
 
-	ViewInitOptions.ViewOrigin = ViewTransform.GetLocation();
-	FRotator ViewRotation = ViewTransform.GetRotation();
+	// Apply view modifiers.
+	FMinimalViewInfo ModifiedViewInfo;
+	{
+		ModifiedViewInfo.Location = ViewTransform.GetLocation();
+		ModifiedViewInfo.Rotation = ViewTransform.GetRotation();
+		ModifiedViewInfo.FOV = ViewFOV;
 
+		if (bShouldApplyViewModifiers)
+		{
+			ViewModifiers.Broadcast(ModifiedViewInfo);
+		}
+	}
+	const FVector ModifiedViewLocation = ModifiedViewInfo.Location;
+	FRotator ModifiedViewRotation = ModifiedViewInfo.Rotation;
+	const float ModifiedViewFOV = ModifiedViewInfo.FOV;
+
+	ViewInitOptions.ViewOrigin = ModifiedViewLocation;
 
 	// Apply head tracking!  Note that this won't affect what the editor *thinks* the view location and rotation is, it will
 	// only affect the rendering of the scene.
@@ -754,13 +769,11 @@ FSceneView* FEditorViewportClient::CalcSceneView(FSceneViewFamily* ViewFamily, c
 		FVector CurrentHmdPosition;
 		GEngine->XRSystem->GetCurrentPose(IXRTrackingSystem::HMDDeviceId, CurrentHmdOrientation, CurrentHmdPosition );
 
-		const FQuat VisualRotation = ViewRotation.Quaternion() * CurrentHmdOrientation;
-		ViewRotation = VisualRotation.Rotator();
-		ViewRotation.Normalize();
+		const FQuat VisualRotation = ModifiedViewRotation.Quaternion() * CurrentHmdOrientation;
+		ModifiedViewRotation = VisualRotation.Rotator();
+		ModifiedViewRotation.Normalize();
 	}
 
-
-
 	FIntPoint ViewportSize = Viewport->GetSizeXY();
 	FIntPoint ViewportOffset(0, 0);
 
@@ -809,7 +822,7 @@ FSceneView* FEditorViewportClient::CalcSceneView(FSceneViewFamily* ViewFamily, c
 	if (bUseControllingActorViewInfo)
 	{
 		// @todo vreditor: Not stereo friendly yet
-		ViewInitOptions.ViewRotationMatrix = FInverseRotationMatrix(ViewRotation) * FMatrix(
+		ViewInitOptions.ViewRotationMatrix = FInverseRotationMatrix(ModifiedViewRotation) * FMatrix(
 			FPlane(0, 0, 1, 0),
 			FPlane(1, 0, 0, 0),
 			FPlane(0, 1, 0, 0),
@@ -834,11 +847,11 @@ FSceneView* FEditorViewportClient::CalcSceneView(FSceneViewFamily* ViewFamily, c
 		        const FIntRect StereoViewRect = FIntRect( X, Y, X + SizeX, Y + SizeY );
 		        ViewInitOptions.SetViewRectangle( StereoViewRect );
 
-				GEngine->StereoRenderingDevice->CalculateStereoViewOffset( StereoPass, ViewRotation, ViewInitOptions.WorldToMetersScale, ViewInitOptions.ViewOrigin );
+				GEngine->StereoRenderingDevice->CalculateStereoViewOffset( StereoPass, ModifiedViewRotation, ViewInitOptions.WorldToMetersScale, ViewInitOptions.ViewOrigin );
 			}
 
 			// Calc view rotation matrix
-			ViewInitOptions.ViewRotationMatrix = CalcViewRotationMatrix(ViewRotation);
+			ViewInitOptions.ViewRotationMatrix = CalcViewRotationMatrix(ModifiedViewRotation);
 
 		    // Rotate view 90 degrees
 			ViewInitOptions.ViewRotationMatrix = ViewInitOptions.ViewRotationMatrix * FMatrix(
@@ -858,7 +871,7 @@ FSceneView* FEditorViewportClient::CalcSceneView(FSceneViewFamily* ViewFamily, c
 			    const float MinZ = GetNearClipPlane();
 			    const float MaxZ = MinZ;
 			    // Avoid zero ViewFOV's which cause divide by zero's in projection matrix
-			    const float MatrixFOV = FMath::Max(0.001f, ViewFOV) * (float)PI / 360.0f;
+			    const float MatrixFOV = FMath::Max(0.001f, ModifiedViewFOV) * (float)PI / 360.0f;
 
 			    if (bConstrainAspectRatio)
 			    {
@@ -1045,7 +1058,7 @@ FSceneView* FEditorViewportClient::CalcSceneView(FSceneViewFamily* ViewFamily, c
 	ViewInitOptions.OverrideLODViewOrigin = FVector::ZeroVector;
 	ViewInitOptions.bUseFauxOrthoViewPos = true;
 
-	ViewInitOptions.FOV = ViewFOV;
+	ViewInitOptions.FOV = ModifiedViewFOV;
 	if (bUseControllingActorViewInfo)
 	{
 		ViewInitOptions.bUseFieldOfViewForLOD = ControllingActorViewInfo.bUseFieldOfViewForLOD;
@@ -1057,8 +1070,8 @@ FSceneView* FEditorViewportClient::CalcSceneView(FSceneViewFamily* ViewFamily, c
 
 	FSceneView* View = new FSceneView(ViewInitOptions);
 
-	View->ViewLocation = ViewTransform.GetLocation();
-	View->ViewRotation = ViewRotation;
+	View->ViewLocation = ModifiedViewLocation;
+	View->ViewRotation = ModifiedViewRotation;
 
 	View->SubduedSelectionOutlineColor = GEngine->GetSubduedSelectionOutlineColor();
 
diff --git a/Engine/Source/Editor/UnrealEd/Private/LevelEditorViewport.cpp b/Engine/Source/Editor/UnrealEd/Private/LevelEditorViewport.cpp
--- a/Engine/Source/Editor/UnrealEd/Private/LevelEditorViewport.cpp
+++ b/Engine/Source/Editor/UnrealEd/Private/LevelEditorViewport.cpp
@@ -1741,7 +1741,6 @@ FLevelEditorViewportClient::FLevelEditorViewportClient(const TSharedPtr<SLevelVi
 	, bLockedCameraView(true)
 	, bReceivedFocusRecently(false)
 	, bAlwaysShowModeWidgetAfterSelectionChanges(true)
-	, bShouldApplyViewModifiers(true)
 	, SpriteCategoryVisibility()
 	, World(nullptr)
 	, TrackingTransaction()
@@ -2290,39 +2289,9 @@ void FLevelEditorViewportClient::Tick(float DeltaTime)
 
 	UpdateViewForLockedActor(DeltaTime);
 
-	if (bShouldApplyViewModifiers)
-	{
-		FViewportCameraTransform& ViewTransform = GetViewTransform();
-		if (!ViewTransform.IsPlaying())
-		{
-			ApplyViewModifiers(DeltaTime);
-		}
-	}
-
 	UserIsControllingAtmosphericLightTimer = FMath::Max(UserIsControllingAtmosphericLightTimer - DeltaTime, 0.0f);
 }
 
-void FLevelEditorViewportClient::ApplyViewModifiers(float DeltaTime)
-{
-	if (ViewModifiers.IsBound())
-	{
-		FEditorViewportViewModifierParams Params;
-		Params.DeltaTime = DeltaTime;
-		Params.ViewportClient = this;
-
-		FMinimalViewInfo InOutPOV;
-		InOutPOV.Location = GetViewLocation();
-		InOutPOV.Rotation = GetViewRotation();
-		InOutPOV.FOV = ViewFOV;
-
-		ViewModifiers.Broadcast(Params, InOutPOV);
-
-		SetViewLocation(InOutPOV.Location);
-		SetViewRotation(InOutPOV.Rotation);
-		ViewFOV = InOutPOV.FOV;
-	}
-}
-
 void FLevelEditorViewportClient::UpdateViewForLockedActor(float DeltaTime)
 {
 	// We can't be locked to a matinee actor if this viewport doesn't allow matinee control
diff --git a/Engine/Source/Editor/UnrealEd/Public/EditorViewportClient.h b/Engine/Source/Editor/UnrealEd/Public/EditorViewportClient.h
--- a/Engine/Source/Editor/UnrealEd/Public/EditorViewportClient.h
+++ b/Engine/Source/Editor/UnrealEd/Public/EditorViewportClient.h
@@ -267,6 +267,9 @@ private:
 	float OrthoZoom;
 };
 
+/** Delegate for modifying view parameters of an editor viewport. */
+DECLARE_MULTICAST_DELEGATE_OneParam(FEditorViewportViewModifierDelegate, FMinimalViewInfo&);
+
 /** Viewport client for editor viewports. Contains common functionality for camera movement, rendering debug information, etc. */
 class UNREALED_API FEditorViewportClient : public FCommonViewportClient, public FViewElementDrawer, public FGCObject
 {
@@ -1492,6 +1495,12 @@ public:
 	/** If true, draw vertices for selected BSP brushes and static meshes if the large vertices show flag is set. */
 	bool bDrawVertices;
 
+	/** List of view modifiers to apply on view parameters. */
+	FEditorViewportViewModifierDelegate ViewModifiers;
+
+	/** Whether view modifiers should be called and applied. */
+	bool bShouldApplyViewModifiers;
+
 protected:
 	/** Does this viewport client own the mode tools instance pointed at by ModeTools control the lifetime of it? */
 	bool bOwnsModeTools;
diff --git a/Engine/Source/Editor/UnrealEd/Public/LevelEditorViewport.h b/Engine/Source/Editor/UnrealEd/Public/LevelEditorViewport.h
--- a/Engine/Source/Editor/UnrealEd/Public/LevelEditorViewport.h
+++ b/Engine/Source/Editor/UnrealEd/Public/LevelEditorViewport.h
@@ -121,17 +121,6 @@ private:
 	
 };
 
-struct FEditorViewportViewModifierParams
-{
-	/** Delta time for updating view modifications. */
-	float DeltaTime = 0.f;
-	/** The viewport client calling for view modification. */
-	FEditorViewportClient* ViewportClient = nullptr;
-};
-
-/** Delegate for modifying view parameters of an editor viewport. */
-DECLARE_MULTICAST_DELEGATE_TwoParams(FEditorViewportViewModifierDelegate, const FEditorViewportViewModifierParams&, FMinimalViewInfo&);
-
 /** */
 class UNREALED_API FLevelEditorViewportClient : public FEditorViewportClient
 {
@@ -580,9 +569,6 @@ protected:
 	/** Delegate handler for ActorMoved events */
 	void OnActorMoved(AActor* InActor);
 
-	/** Call registered view modifiers and apply the modify result */
-	void ApplyViewModifiers(float DeltaTime);
-
 	/** FEditorViewportClient Interface*/
 
 public:
@@ -763,12 +749,6 @@ public:
 	/** When enabled, the Unreal transform widget will become visible after an actor is selected, even if it was turned off via a show flag */
 	bool bAlwaysShowModeWidgetAfterSelectionChanges;
 
-	/** List of view modifiers to apply on view parameters. */
-	FEditorViewportViewModifierDelegate ViewModifiers;
-
-	/** Whether view modifiers should be called and applied. */
-	bool bShouldApplyViewModifiers;
-
 private:
 	/** The actors that are currently being placed in the viewport via dragging */
 	static TArray< TWeakObjectPtr< AActor > > DropPreviewActors;
