Fix non-unity builds

diff --git a/Engine/Plugins/Cameras/CameraShakePreviewer/Source/CameraShakePreviewer/Private/SCameraShakePreviewer.cpp b/Engine/Plugins/Cameras/CameraShakePreviewer/Source/CameraShakePreviewer/Private/SCameraShakePreviewer.cpp
--- a/Engine/Plugins/Cameras/CameraShakePreviewer/Source/CameraShakePreviewer/Private/SCameraShakePreviewer.cpp
+++ b/Engine/Plugins/Cameras/CameraShakePreviewer/Source/CameraShakePreviewer/Private/SCameraShakePreviewer.cpp
@@ -10,7 +10,6 @@
 #include "LevelEditorViewport.h"
 #include "Modules/ModuleManager.h"
 #include "Slate/SceneViewport.h"
-#include "TickableEditorObject.h"
 #include "Widgets/Input/SButton.h"
 #include "Widgets/Input/SCheckBox.h"
 #include "Widgets/Layout/SBox.h"
diff --git a/Engine/Plugins/Cameras/CameraShakePreviewer/Source/CameraShakePreviewer/Private/SCameraShakePreviewer.h b/Engine/Plugins/Cameras/CameraShakePreviewer/Source/CameraShakePreviewer/Private/SCameraShakePreviewer.h
--- a/Engine/Plugins/Cameras/CameraShakePreviewer/Source/CameraShakePreviewer/Private/SCameraShakePreviewer.h
+++ b/Engine/Plugins/Cameras/CameraShakePreviewer/Source/CameraShakePreviewer/Private/SCameraShakePreviewer.h
@@ -5,6 +5,7 @@
 #include "Containers/Array.h"
 #include "EditorUndoClient.h"
 #include "Templates/SharedPointer.h"
+#include "TickableEditorObject.h"
 #include "UObject/GCObject.h"
 #include "Widgets/DeclarativeSyntaxSupport.h"
 #include "Widgets/Views/SListView.h"
