Cameras: Detect when shakes with duration end, and mark them as stopped in previewer panel.

diff --git a/Engine/Plugins/Cameras/CameraShakePreviewer/Source/CameraShakePreviewer/Private/SCameraShakePreviewer.cpp b/Engine/Plugins/Cameras/CameraShakePreviewer/Source/CameraShakePreviewer/Private/SCameraShakePreviewer.cpp
--- a/Engine/Plugins/Cameras/CameraShakePreviewer/Source/CameraShakePreviewer/Private/SCameraShakePreviewer.cpp
+++ b/Engine/Plugins/Cameras/CameraShakePreviewer/Source/CameraShakePreviewer/Private/SCameraShakePreviewer.cpp
@@ -420,6 +420,8 @@ void SCameraShakePreviewer::Tick(const FGeometry& AllottedGeometry, const double
 	}
 
 	// Update playing information on our camera shakes.
+    TArray<FActiveCameraShakeInfo> ActiveCameraShakes;
+    CameraShakePreviewUpdater->ShakeModifier().GetActiveCameraShakes(ActiveCameraShakes);
 	for (TSharedPtr<FCameraShakeData> CameraShake : CameraShakes)
 	{
 		if (CameraShake->SourceComponent.IsValid())
@@ -461,6 +463,19 @@ void SCameraShakePreviewer::Tick(const FGeometry& AllottedGeometry, const double
 					}
 				}
 			}
+
+            // Check if the shake is still playing.
+            if (CameraShake->bIsPlaying && CameraShake->ShakeInstance != nullptr)
+            {
+                if (!ActiveCameraShakes.ContainsByPredicate(
+                        [CameraShake](const FActiveCameraShakeInfo& ShakeInfo)
+                        {
+                            return ShakeInfo.ShakeInstance == CameraShake->ShakeInstance;
+                        }))
+                {
+                    CameraShake->bIsPlaying = false;
+                }
+            }
 		}
 	}
 
