Camera shake toggle converted to use tool menu API

diff --git a/Engine/Plugins/Cameras/CameraShakePreviewer/Source/CameraShakePreviewer/Private/CameraShakePreviewerModule.cpp b/Engine/Plugins/Cameras/CameraShakePreviewer/Source/CameraShakePreviewer/Private/CameraShakePreviewerModule.cpp
--- a/Engine/Plugins/Cameras/CameraShakePreviewer/Source/CameraShakePreviewer/Private/CameraShakePreviewerModule.cpp
+++ b/Engine/Plugins/Cameras/CameraShakePreviewer/Source/CameraShakePreviewer/Private/CameraShakePreviewerModule.cpp
@@ -18,6 +18,7 @@
 #include "WorkspaceMenuStructureModule.h"
 #include "IAssetViewport.h"
 #include "EditorViewportClient.h"
+#include "LevelEditorMenuContext.h"
 
 #define LOCTEXT_NAMESPACE "CameraShakePreviewer"
 
@@ -106,6 +107,69 @@ void FCameraShakePreviewerModule::UnregisterEditorTab()
 
 void FCameraShakePreviewerModule::RegisterViewportOptionMenuExtender()
 {
+	{
+		FToolMenuOwnerScoped ToolMenuOwnerScoped(this);
+
+		{
+			UToolMenu* Menu = UToolMenus::Get()->ExtendMenu("LevelEditor.LevelViewportToolBar.Options");
+
+			{
+				FToolMenuSection& Section = Menu->FindOrAddSection("LevelViewportViewportOptions2");
+
+				// ToggleCameraShakesPreview
+				{
+					static auto GetPerspectiveLevelEditorViewportClient = [](const FToolMenuContext& MenuContext) -> FLevelEditorViewportClient*
+					{
+						ULevelViewportToolBarContext* Context = MenuContext.FindContext<ULevelViewportToolBarContext>();
+						if (Context && Context->LevelViewportToolBarWidget.IsValid())
+						{
+							FLevelEditorViewportClient* ViewportClient = Context->GetLevelViewportClient();
+							if (ViewportClient && ViewportClient->ViewportType == ELevelViewportType::LVT_Perspective)
+							{
+								return ViewportClient;
+							}
+						}
+
+						return nullptr;
+					};
+
+					FToolUIAction Action;
+
+					Action.ExecuteAction.BindLambda([this](const FToolMenuContext& MenuContext)
+					{
+						if (FLevelEditorViewportClient* ViewportClient = GetPerspectiveLevelEditorViewportClient(MenuContext))
+						{
+							ToggleCameraShakesPreview(ViewportClient);
+						}
+					});
+
+					Action.GetActionCheckState.BindLambda([this](const FToolMenuContext& MenuContext) -> ECheckBoxState
+					{
+						if (FLevelEditorViewportClient* ViewportClient = GetPerspectiveLevelEditorViewportClient(MenuContext))
+						{
+							if (HasCameraShakesPreview(ViewportClient))
+							{
+								return ECheckBoxState::Checked;
+							}
+						}
+
+						return ECheckBoxState::Unchecked;
+					});
+
+					TSharedPtr<FUICommandInfo> ToggleCameraShakesPreview = FCameraShakePreviewerCommands::Get().ToggleCameraShakesPreview;
+					Section.AddMenuEntry(
+						"ToggleCameraShakesPreview",
+						ToggleCameraShakesPreview->GetLabel(),
+						ToggleCameraShakesPreview->GetDescription(),
+						ToggleCameraShakesPreview->GetIcon(),
+						Action,
+						ToggleCameraShakesPreview->GetUserInterfaceType()
+					);
+				}
+			}
+		}
+	}
+
 	// Register a callback for adding a "Show Camera Shakes" option in the viewport options menu.
 	FLevelEditorModule& LevelEditorModule = FModuleManager::LoadModuleChecked<FLevelEditorModule>(LevelEditorModuleName);
 
@@ -120,6 +184,8 @@ void FCameraShakePreviewerModule::RegisterViewportOptionMenuExtender()
 
 void FCameraShakePreviewerModule::UnregisterViewportOptionMenuExtender()
 {
+	UToolMenus::UnregisterOwner(this);
+
 	FLevelEditorModule& LevelEditorModule = FModuleManager::GetModuleChecked<FLevelEditorModule>(LevelEditorModuleName);
 	LevelEditorModule.GetAllLevelViewportOptionsMenuExtenders().RemoveAll([=](const FLevelEditorModule::FLevelEditorMenuExtender& Extender) { return Extender.GetHandle() == ViewportOptionsMenuExtenderHandle; });
 
@@ -155,21 +221,6 @@ TSharedRef<FExtender> FCameraShakePreviewerModule::OnExtendLevelViewportOptionMe
 
 	// Make the extender that adds the menu entry.
 	TSharedRef<FExtender> Extender = MakeShared<FExtender>();
-	Extender->AddMenuExtension(
-		"LevelViewportViewportOptions2",
-		EExtensionHook::After,
-		nullptr,
-		FMenuExtensionDelegate::CreateLambda([ToggleCameraShakesPreviewAction](FMenuBuilder& MenuBuilder)
-		{
-			MenuBuilder.AddMenuEntry(
-				LOCTEXT("AllowCameraShakes", "Allow Camera Shakes"),
-				LOCTEXT("AllowCameraShakesTooltip", "If enabled, allows the camera shakes previewer panel to apply shakes to this viewport"),
-				FSlateIcon(),
-				ToggleCameraShakesPreviewAction,
-				NAME_None,
-				EUserInterfaceActionType::ToggleButton);
-		})
-	);
 	return Extender;
 }
 
diff --git a/Engine/Source/Editor/LevelEditor/Private/SLevelViewportToolBar.cpp b/Engine/Source/Editor/LevelEditor/Private/SLevelViewportToolBar.cpp
--- a/Engine/Source/Editor/LevelEditor/Private/SLevelViewportToolBar.cpp
+++ b/Engine/Source/Editor/LevelEditor/Private/SLevelViewportToolBar.cpp
@@ -311,6 +311,15 @@ bool SLevelViewportToolBar::IsViewModeSupported(EViewModeIndex ViewModeIndex) co
 	return true;
 }
 
+FLevelEditorViewportClient* SLevelViewportToolBar::GetLevelViewportClient()
+{
+	if (Viewport.IsValid())
+	{
+		return &Viewport.Pin()->GetLevelViewportClient();
+	}
+
+	return nullptr;
+}
 
 FText SLevelViewportToolBar::GetCameraMenuLabel() const
 {
@@ -1582,4 +1591,14 @@ TSharedRef<SWidget> SLevelViewportToolBar::GetScalabilityWarningMenuContent() co
 		];
 }
 
+FLevelEditorViewportClient* ULevelViewportToolBarContext::GetLevelViewportClient()
+{
+	if (LevelViewportToolBarWidget.IsValid())
+	{
+		return LevelViewportToolBarWidget.Pin()->GetLevelViewportClient();
+	}
+
+	return nullptr;
+}
+
 #undef LOCTEXT_NAMESPACE
diff --git a/Engine/Source/Editor/LevelEditor/Private/SLevelViewportToolBar.h b/Engine/Source/Editor/LevelEditor/Private/SLevelViewportToolBar.h
--- a/Engine/Source/Editor/LevelEditor/Private/SLevelViewportToolBar.h
+++ b/Engine/Source/Editor/LevelEditor/Private/SLevelViewportToolBar.h
@@ -30,6 +30,9 @@ public:
 	/** @return Whether the given viewmode is supported. */ 
 	virtual bool IsViewModeSupported(EViewModeIndex ViewModeIndex) const;
 
+	/** @return Level editor viewport client. */ 
+	FLevelEditorViewportClient* GetLevelViewportClient();
+
 private:
 	/**
 	 * Returns the label for the "Camera" tool bar menu, which changes depending on the viewport type
diff --git a/Engine/Source/Editor/LevelEditor/Public/LevelEditorMenuContext.h b/Engine/Source/Editor/LevelEditor/Public/LevelEditorMenuContext.h
--- a/Engine/Source/Editor/LevelEditor/Public/LevelEditorMenuContext.h
+++ b/Engine/Source/Editor/LevelEditor/Public/LevelEditorMenuContext.h
@@ -48,4 +48,6 @@ class LEVELEDITOR_API ULevelViewportToolBarContext : public UObject
 public:
 	TWeakPtr<SLevelViewportToolBar> LevelViewportToolBarWidget;
 	TWeakPtr<const SLevelViewportToolBar> LevelViewportToolBarWidgetConst;
+
+	FLevelEditorViewportClient* GetLevelViewportClient();
 };
