CameraShakePreviewer: stop and discard camera shakes that are getting recompiled

diff --git a/Engine/Plugins/Cameras/CameraShakePreviewer/Source/CameraShakePreviewer/Private/SCameraShakePreviewer.cpp b/Engine/Plugins/Cameras/CameraShakePreviewer/Source/CameraShakePreviewer/Private/SCameraShakePreviewer.cpp
--- a/Engine/Plugins/Cameras/CameraShakePreviewer/Source/CameraShakePreviewer/Private/SCameraShakePreviewer.cpp
+++ b/Engine/Plugins/Cameras/CameraShakePreviewer/Source/CameraShakePreviewer/Private/SCameraShakePreviewer.cpp
@@ -77,8 +77,27 @@
 {
 	PreviewCameraShake = CastChecked<UCameraModifier_CameraShake>(
 			PreviewCamera->AddNewCameraModifier(UCameraModifier_CameraShake::StaticClass()));
+
+	// Handle camera shakes being recompiled.
+	FCoreUObjectDelegates::OnObjectsReplaced.AddRaw(this, &FCameraShakePreviewUpdater::OnObjectsReplaced);
 }
 
+void FCameraShakePreviewUpdater::OnObjectsReplaced(const TMap<UObject*, UObject*>& ReplacementMap)
+{
+	// If a camera shake gets recompiled, we just stop and discard it.
+	TArray<FActiveCameraShakeInfo> ActiveCameraShakes;
+	PreviewCameraShake->GetActiveCameraShakes(ActiveCameraShakes);
+
+	for (const FActiveCameraShakeInfo& ActiveCameraShake : ActiveCameraShakes)
+	{
+		UObject* NewShakeInstance = ReplacementMap.FindRef(ActiveCameraShake.ShakeInstance);
+		if (NewShakeInstance)
+		{
+			PreviewCameraShake->RemoveCameraShake(ActiveCameraShake.ShakeInstance, true);
+		}
+	}
+}
+
 void FCameraShakePreviewUpdater::Tick(float DeltaTime)
 {
 	// Accumulate the deltas in case we get ticked several times before we are 
diff --git a/Engine/Plugins/Cameras/CameraShakePreviewer/Source/CameraShakePreviewer/Private/SCameraShakePreviewer.h b/Engine/Plugins/Cameras/CameraShakePreviewer/Source/CameraShakePreviewer/Private/SCameraShakePreviewer.h
--- a/Engine/Plugins/Cameras/CameraShakePreviewer/Source/CameraShakePreviewer/Private/SCameraShakePreviewer.h
+++ b/Engine/Plugins/Cameras/CameraShakePreviewer/Source/CameraShakePreviewer/Private/SCameraShakePreviewer.h
@@ -79,6 +79,8 @@
 
 	void AddPostProcessBlend(const FPostProcessSettings& Settings, float Weight);
 
+	void OnObjectsReplaced(const TMap<UObject*, UObject*>& ReplacementMap);
+
 private:
 	APreviewPlayerCameraManager* PreviewCamera;
 	UCameraModifier_CameraShake* PreviewCameraShake;
