Cameras: fix leaked engine callback when the camera shake previewer is closed

diff --git a/Engine/Plugins/Cameras/CameraShakePreviewer/Source/CameraShakePreviewer/Private/SCameraShakePreviewer.cpp b/Engine/Plugins/Cameras/CameraShakePreviewer/Source/CameraShakePreviewer/Private/SCameraShakePreviewer.cpp
--- a/Engine/Plugins/Cameras/CameraShakePreviewer/Source/CameraShakePreviewer/Private/SCameraShakePreviewer.cpp
+++ b/Engine/Plugins/Cameras/CameraShakePreviewer/Source/CameraShakePreviewer/Private/SCameraShakePreviewer.cpp
@@ -82,6 +82,11 @@ FCameraShakePreviewUpdater::FCameraShakePreviewUpdater()
 	FCoreUObjectDelegates::OnObjectsReplaced.AddRaw(this, &FCameraShakePreviewUpdater::OnObjectsReplaced);
 }
 
+FCameraShakePreviewUpdater::~FCameraShakePreviewUpdater()
+{
+	FCoreUObjectDelegates::OnObjectsReplaced.RemoveAll(this);
+}
+
 void FCameraShakePreviewUpdater::OnObjectsReplaced(const TMap<UObject*, UObject*>& ReplacementMap)
 {
 	// If a camera shake gets recompiled, we just stop and discard it.
diff --git a/Engine/Plugins/Cameras/CameraShakePreviewer/Source/CameraShakePreviewer/Private/SCameraShakePreviewer.h b/Engine/Plugins/Cameras/CameraShakePreviewer/Source/CameraShakePreviewer/Private/SCameraShakePreviewer.h
--- a/Engine/Plugins/Cameras/CameraShakePreviewer/Source/CameraShakePreviewer/Private/SCameraShakePreviewer.h
+++ b/Engine/Plugins/Cameras/CameraShakePreviewer/Source/CameraShakePreviewer/Private/SCameraShakePreviewer.h
@@ -54,6 +54,7 @@ class FCameraShakePreviewUpdater : public FTickableEditorObject, public FGCObjec
 {
 public:
 	FCameraShakePreviewUpdater();
+	~FCameraShakePreviewUpdater();
 
 	// FTickableObject Interface
 	virtual ETickableTickType GetTickableTickType() const { return ETickableTickType::Always; }
