Cameras: fix leaked engine callback when the camera shake previewer is closed

diff --git a/Engine/Plugins/Cameras/CameraShakePreviewer/Source/CameraShakePreviewer/Private/SCameraShakePreviewer.h b/Engine/Plugins/Cameras/CameraShakePreviewer/Source/CameraShakePreviewer/Private/SCameraShakePreviewer.h
--- a/Engine/Plugins/Cameras/CameraShakePreviewer/Source/CameraShakePreviewer/Private/SCameraShakePreviewer.h
+++ b/Engine/Plugins/Cameras/CameraShakePreviewer/Source/CameraShakePreviewer/Private/SCameraShakePreviewer.h
@@ -54,7 +54,7 @@ class FCameraShakePreviewUpdater : public FTickableEditorObject, public FGCObjec
 {
 public:
 	FCameraShakePreviewUpdater();
-	~FCameraShakePreviewUpdater();
+	virtual ~FCameraShakePreviewUpdater();
 
 	// FTickableObject Interface
 	virtual ETickableTickType GetTickableTickType() const { return ETickableTickType::Always; }
