Sequencer: remove code duplication for previewing camera shakes in editor, fix some crashes when shakes are previewed

- Both the camera shake previewer tool and the camera shake tracks had some form of previewer utility class. The former supported post-process settings, so that code has been merged into the latter, which is the one we keep because of module dependencies.
- The shake previewer needs to properly spawn the temp player camera manager in the world, otherwise some shakes might crash with null world pointers. This is done with a new Initialize() method.
- Better handle map/level changes now that the player camera manager is spawned.

diff --git a/Engine/Plugins/Cameras/CameraShakePreviewer/Source/CameraShakePreviewer/CameraShakePreviewer.Build.cs b/Engine/Plugins/Cameras/CameraShakePreviewer/Source/CameraShakePreviewer/CameraShakePreviewer.Build.cs
--- a/Engine/Plugins/Cameras/CameraShakePreviewer/Source/CameraShakePreviewer/CameraShakePreviewer.Build.cs
+++ b/Engine/Plugins/Cameras/CameraShakePreviewer/Source/CameraShakePreviewer/CameraShakePreviewer.Build.cs
@@ -22,6 +22,8 @@ public class CameraShakePreviewer : ModuleRules
 				"AppFramework",
 				"GameplayCameras",
 				"LevelEditor",
+				"MovieScene",
+				"MovieSceneTracks",
 				"PropertyEditor",
 				"Slate",
 				"ToolMenus",
diff --git a/Engine/Plugins/Cameras/CameraShakePreviewer/Source/CameraShakePreviewer/Private/SCameraShakePreviewer.cpp b/Engine/Plugins/Cameras/CameraShakePreviewer/Source/CameraShakePreviewer/Private/SCameraShakePreviewer.cpp
--- a/Engine/Plugins/Cameras/CameraShakePreviewer/Source/CameraShakePreviewer/Private/SCameraShakePreviewer.cpp
+++ b/Engine/Plugins/Cameras/CameraShakePreviewer/Source/CameraShakePreviewer/Private/SCameraShakePreviewer.cpp
@@ -6,6 +6,7 @@
 #include "Camera/CameraShakeSourceActor.h"
 #include "Camera/CameraShakeSourceComponent.h"
 #include "CameraShakePreviewerModule.h"
+#include "EditorSupportDelegates.h"
 #include "Engine/Level.h"
 #include "LevelEditorViewport.h"
 #include "Framework/Views/TableViewMetadata.h"
@@ -14,13 +15,9 @@
 #include "Widgets/Layout/SWrapBox.h"
 #include "Widgets/Views/SListView.h"
 
-#include UE_INLINE_GENERATED_CPP_BY_NAME(SCameraShakePreviewer)
-
 #define LOCTEXT_NAMESPACE "CameraShakePreviewer"
 
-namespace UE
-{
-namespace MovieScene
+namespace UE::MovieScene
 {
 
 UWorld* FindCameraShakePreviewerWorld()
@@ -35,8 +32,7 @@ UWorld* FindCameraShakePreviewerWorld()
 	return nullptr;
 }
 
-}
-}
+} // namespace UE::MovieScene
 
 /**
  * Data struct for each entry in the panel's main list.
@@ -63,121 +59,100 @@ struct FCameraShakeData
 };
 
 FCameraShakePreviewUpdater::FCameraShakePreviewUpdater()
-	: PreviewCamera(NewObject<APreviewPlayerCameraManager>())
-	, LastLocationModifier(FVector::ZeroVector)
-	, LastRotationModifier(FRotator::ZeroRotator)
-	, LastFOVModifier(0.f)
 {
-	PreviewCameraShake = CastChecked<UCameraModifier_CameraShake>(
-			PreviewCamera->AddNewCameraModifier(UCameraModifier_CameraShake::StaticClass()));
-
 	// Handle camera shakes being recompiled.
 	FCoreUObjectDelegates::OnObjectsReplaced.AddRaw(this, &FCameraShakePreviewUpdater::OnObjectsReplaced);
+
+	// Handle levels/worlds being cleared.
+	FEditorSupportDelegates::PrepareToCleanseEditorObject.AddRaw(this, &FCameraShakePreviewUpdater::OnPrepareToCleanseEditorObject);
 }
 
 FCameraShakePreviewUpdater::~FCameraShakePreviewUpdater()
 {
+	FEditorSupportDelegates::PrepareToCleanseEditorObject.RemoveAll(this);
 	FCoreUObjectDelegates::OnObjectsReplaced.RemoveAll(this);
 }
 
-void FCameraShakePreviewUpdater::OnObjectsReplaced(const TMap<UObject*, UObject*>& ReplacementMap)
+void FCameraShakePreviewUpdater::Reinitialize(UWorld* InWorld)
 {
-	// If a camera shake gets recompiled, we just stop and discard it.
-	TArray<FActiveCameraShakeInfo> ActiveCameraShakes;
-	PreviewCameraShake->GetActiveCameraShakes(ActiveCameraShakes);
-
-	for (const FActiveCameraShakeInfo& ActiveCameraShake : ActiveCameraShakes)
+	if (Previewer.IsInitialized() && Previewer.GetCameraManager()->GetWorld() != InWorld)
 	{
-		UObject* NewShakeInstance = ReplacementMap.FindRef(ActiveCameraShake.ShakeInstance);
-		if (NewShakeInstance)
-		{
-			PreviewCameraShake->RemoveCameraShake(ActiveCameraShake.ShakeInstance, true);
-		}
+		Previewer.Teardown();
+	}
+	if (!Previewer.IsInitialized())
+	{
+		Previewer.Initialize(InWorld);
 	}
 }
 
-void FCameraShakePreviewUpdater::Tick(float DeltaTime)
-{
-	// Accumulate the deltas in case we get ticked several times before we are 
-	// asked to modify a viewport.
-	LastDeltaTime = LastDeltaTime.Get(0.f) + DeltaTime; 
-}
-
-void FCameraShakePreviewUpdater::ModifyCamera(FEditorViewportViewModifierParams& Params)
+void FCameraShakePreviewUpdater::OnObjectsReplaced(const TMap<UObject*, UObject*>& ReplacementMap)
 {
-	const float DeltaTime = LastDeltaTime.Get(-1.f);
-	if (DeltaTime > 0.f)
+	if (Previewer.IsInitialized())
 	{
-		LastPostProcessSettings.Reset();
-		LastPostProcessBlendWeights.Reset();
-		PreviewCamera->ResetPostProcessSettings();
-
-		FMinimalViewInfo OriginalPOV(Params.ViewInfo);
-
-		PreviewCameraShake->ModifyCamera(DeltaTime, Params.ViewInfo);
-
-		PreviewCamera->MergePostProcessSettings(LastPostProcessSettings, LastPostProcessBlendWeights);
+		// If a camera shake gets recompiled, we just stop and discard it.
+		TArray<FActiveCameraShakeInfo> ActiveCameraShakes;
+		Previewer.GetCameraModifier()->GetActiveCameraShakes(ActiveCameraShakes);
 
-		LastLocationModifier = Params.ViewInfo.Location - OriginalPOV.Location;
-		LastRotationModifier = Params.ViewInfo.Rotation - OriginalPOV.Rotation;
-		LastFOVModifier = Params.ViewInfo.FOV - OriginalPOV.FOV;
-
-		LastDeltaTime.Reset();
-	}
-	else
-	{
-		Params.ViewInfo.Location += LastLocationModifier;
-		Params.ViewInfo.Rotation += LastRotationModifier;
-		Params.ViewInfo.FOV += LastFOVModifier;
+		for (const FActiveCameraShakeInfo& ActiveCameraShake : ActiveCameraShakes)
+		{
+			UObject* NewShakeInstance = ReplacementMap.FindRef(ActiveCameraShake.ShakeInstance);
+			if (NewShakeInstance)
+			{
+				Previewer.GetCameraModifier()->RemoveCameraShake(ActiveCameraShake.ShakeInstance, true);
+			}
+		}
 	}
+}
 
-	for (int32 PPIndex = 0; PPIndex < LastPostProcessSettings.Num(); ++PPIndex)
+void FCameraShakePreviewUpdater::OnPrepareToCleanseEditorObject(UObject* Object)
+{
+	if (Previewer.IsInitialized())
 	{
-		Params.AddPostProcessBlend(LastPostProcessSettings[PPIndex], LastPostProcessBlendWeights[PPIndex]);
+		ULevel* Level = Previewer.GetCameraManager()->GetLevel();
+		if (Object == Level || Object == Level->OwningWorld)
+		{
+			Previewer.Teardown();
+		}
 	}
 }
 
+void FCameraShakePreviewUpdater::Tick(float DeltaTime)
+{
+	Previewer.Update(DeltaTime, true);
+}
 
-
-void FCameraShakePreviewUpdater::AddReferencedObjects(FReferenceCollector& Collector)
+void FCameraShakePreviewUpdater::ModifyCamera(FEditorViewportViewModifierParams& Params)
 {
-	Collector.AddReferencedObject(PreviewCamera);
+	Previewer.ModifyView(Params);
 }
 
 UCameraShakeBase* FCameraShakePreviewUpdater::AddCameraShake(TSubclassOf<UCameraShakeBase> ShakeClass, const FAddCameraShakeParams& Params)
 {
 	FAddCameraShakeParams ActualParams(Params);
 
-	UCameraShakeBase* ShakeInstance = PreviewCameraShake->AddCameraShake(ShakeClass, ActualParams);
+	UCameraShakeBase* ShakeInstance = Previewer.GetCameraModifier()->AddCameraShake(ShakeClass, ActualParams);
 
 	return ShakeInstance;
 }
 
 void FCameraShakePreviewUpdater::RemoveAllCameraShakesFromSource(const UCameraShakeSourceComponent* SourceComponent)
 {
-	PreviewCameraShake->RemoveAllCameraShakesFromSource(SourceComponent);
+	Previewer.GetCameraModifier()->RemoveAllCameraShakesFromSource(SourceComponent);
 }
 
 void FCameraShakePreviewUpdater::GetActiveCameraShakes(TArray<FActiveCameraShakeInfo>& ActiveCameraShakes) const
 {
-	PreviewCameraShake->GetActiveCameraShakes(ActiveCameraShakes);
+	Previewer.GetCameraModifier()->GetActiveCameraShakes(ActiveCameraShakes);
 }
 
 void FCameraShakePreviewUpdater::RemoveCameraShake(UCameraShakeBase* ShakeInstance)
 {
-	PreviewCameraShake->RemoveCameraShake(ShakeInstance);
+	Previewer.GetCameraModifier()->RemoveCameraShake(ShakeInstance);
 }
 
 void FCameraShakePreviewUpdater::RemoveAllCameraShakes()
 {
-	PreviewCameraShake->RemoveAllCameraShakes();
-}
-
-void FCameraShakePreviewUpdater::AddPostProcessBlend(const FPostProcessSettings& Settings, float Weight)
-{
-	check(LastPostProcessSettings.Num() == LastPostProcessBlendWeights.Num());
-	LastPostProcessSettings.Add(Settings);
-	LastPostProcessBlendWeights.Add(Weight);
+	Previewer.GetCameraModifier()->RemoveAllCameraShakes();
 }
 
 /**
@@ -425,12 +400,15 @@ TSharedRef<ITableRow> SCameraShakePreviewer::OnCameraShakesListGenerateRowWidget
 
 void SCameraShakePreviewer::Populate()
 {
-	const bool bFoundWorld = FindCurrentWorld();
-	if (!ensureMsgf(bFoundWorld, TEXT("Could not find current world instance.")))
+	UWorld* CurrentWorld = FindCurrentWorld();
+	if (!ensureMsgf(CurrentWorld, TEXT("Could not find current world instance.")))
 	{
 		return;
 	}
 
+	// Reinitialize the previewer in case the persistent level has changed.
+	CameraShakePreviewUpdater->Reinitialize(CurrentWorld);
+
 	// Get all the shake sources from the level.
 	TArray<UCameraShakeSourceComponent*> ShakeSourceComponents;
 	for (ULevel* Level : CurrentWorld->GetLevels())
@@ -501,10 +479,10 @@ void SCameraShakePreviewer::Populate()
 	CameraShakesListView->RequestListRefresh();
 }
 
-bool SCameraShakePreviewer::FindCurrentWorld()
+UWorld* SCameraShakePreviewer::FindCurrentWorld()
 {
-	CurrentWorld = UE::MovieScene::FindCameraShakePreviewerWorld();
-	return CurrentWorld.IsValid();
+	WeakCurrentWorld = UE::MovieScene::FindCameraShakePreviewerWorld();
+	return WeakCurrentWorld.Get();
 }
 
 void SCameraShakePreviewer::Refresh()
diff --git a/Engine/Plugins/Cameras/CameraShakePreviewer/Source/CameraShakePreviewer/Private/SCameraShakePreviewer.h b/Engine/Plugins/Cameras/CameraShakePreviewer/Source/CameraShakePreviewer/Private/SCameraShakePreviewer.h
--- a/Engine/Plugins/Cameras/CameraShakePreviewer/Source/CameraShakePreviewer/Private/SCameraShakePreviewer.h
+++ b/Engine/Plugins/Cameras/CameraShakePreviewer/Source/CameraShakePreviewer/Private/SCameraShakePreviewer.h
@@ -2,12 +2,11 @@
 
 #pragma once 
 
-#include "Camera/PlayerCameraManager.h"
 #include "EditorUndoClient.h"
 #include "TickableEditorObject.h"
 #include "UObject/GCObject.h"
 #include "Widgets/SCompoundWidget.h"
-#include "SCameraShakePreviewer.generated.h"
+#include "Evaluation/MovieSceneCameraShakePreviewer.h"
 
 class SButton;
 class UCameraModifier_CameraShake;
@@ -26,32 +25,19 @@ class FSceneView;
 class ITableRow;
 class STableViewBase;
 class UCameraAnimInst;
+class UCameraShakeBase;
 class UCameraShakeSourceComponent;
 class ULevel;
 class UWorld;
 struct FCameraShakeData;
 struct FEditorViewportViewModifierParams;
+struct FPostProcessSettings;
 struct FTogglePreviewCameraShakesParams;
 
-UCLASS()
-class APreviewPlayerCameraManager : public APlayerCameraManager
-{
-	GENERATED_BODY()
-
-public:
-	void ResetPostProcessSettings()
-	{
-		ClearCachedPPBlends();
-	}
-
-	void MergePostProcessSettings(TArray<FPostProcessSettings>& InSettings, TArray<float>& InBlendWeights)
-	{
-		InSettings.Append(PostProcessBlendCache);
-		InBlendWeights.Append(PostProcessBlendCacheWeights);
-	}
-};
-
-class FCameraShakePreviewUpdater : public FTickableEditorObject, public FGCObject
+/**
+ * Wrapper for the shake previewer that will tick it with the editor.
+ */
+class FCameraShakePreviewUpdater : public FTickableEditorObject
 {
 public:
 	FCameraShakePreviewUpdater();
@@ -62,35 +48,22 @@ public:
 	virtual TStatId GetStatId() const override { RETURN_QUICK_DECLARE_CYCLE_STAT(FCameraShakePreviewUpdater, STATGROUP_Tickables); }
 	virtual void Tick(float DeltaTime) override;
 
-	// FGCObject interface
-	virtual void AddReferencedObjects(FReferenceCollector& Collector) override;
-	virtual FString GetReferencerName() const override { return TEXT("SCameraShakePreviewer"); }
-
 	UCameraShakeBase* AddCameraShake(TSubclassOf<UCameraShakeBase> ShakeClass, const FAddCameraShakeParams& Params);
-	void RemoveAllCameraShakesFromSource(const UCameraShakeSourceComponent* SourceComponent);
 	void GetActiveCameraShakes(TArray<FActiveCameraShakeInfo>& ActiveCameraShakes) const;
 	void RemoveCameraShake(UCameraShakeBase* ShakeInstance);
+	void RemoveAllCameraShakesFromSource(const UCameraShakeSourceComponent* SourceComponent);
 	void RemoveAllCameraShakes();
 
 	void ModifyCamera(FEditorViewportViewModifierParams& Params);
 
-private:
-	void AddPostProcessBlend(const FPostProcessSettings& Settings, float Weight);
+	void Reinitialize(UWorld* InWorld);
 
+private:
 	void OnObjectsReplaced(const TMap<UObject*, UObject*>& ReplacementMap);
+	void OnPrepareToCleanseEditorObject(UObject* Object);
 
 private:
-	APreviewPlayerCameraManager* PreviewCamera;
-	UCameraModifier_CameraShake* PreviewCameraShake;
-
-	TOptional<float> LastDeltaTime;
-
-	FVector LastLocationModifier;
-	FRotator LastRotationModifier;
-	float LastFOVModifier;
-
-	TArray<FPostProcessSettings> LastPostProcessSettings;
-	TArray<float> LastPostProcessBlendWeights;
+	FCameraShakePreviewer Previewer;
 };
 
 /**
@@ -114,7 +87,7 @@ public:
 
 private:
 	void Populate();
-	bool FindCurrentWorld();
+	UWorld* FindCurrentWorld();
 	void Refresh();
 
 	void OnTogglePreviewCameraShakes(const FTogglePreviewCameraShakesParams& Params);
@@ -151,6 +124,6 @@ private:
 	FLevelEditorViewportClient* ActiveViewportClient;
 	int ActiveViewportIndex;
 
-	TWeakObjectPtr<UWorld> CurrentWorld;
+	TWeakObjectPtr<UWorld> WeakCurrentWorld;
 	bool bNeedsRefresh;
 };
diff --git a/Engine/Source/Runtime/MovieSceneTracks/Private/Evaluation/MovieSceneCameraShakePreviewer.cpp b/Engine/Source/Runtime/MovieSceneTracks/Private/Evaluation/MovieSceneCameraShakePreviewer.cpp
--- a/Engine/Source/Runtime/MovieSceneTracks/Private/Evaluation/MovieSceneCameraShakePreviewer.cpp
+++ b/Engine/Source/Runtime/MovieSceneTracks/Private/Evaluation/MovieSceneCameraShakePreviewer.cpp
@@ -7,8 +7,11 @@
 #include "LevelEditorViewport.h"
 #include "Camera/CameraModifier_CameraShake.h"
 
+#include UE_INLINE_GENERATED_CPP_BY_NAME(MovieSceneCameraShakePreviewer)
+
 FCameraShakePreviewer::FCameraShakePreviewer()
-	: PreviewCameraShake(NewObject<UCameraModifier_CameraShake>())
+	: PreviewCamera(nullptr)
+	, PreviewCameraShake(nullptr)
 	, LastLocationModifier(FVector::ZeroVector)
 	, LastRotationModifier(FRotator::ZeroRotator)
 	, LastFOVModifier(0.f)
@@ -22,7 +25,33 @@ FCameraShakePreviewer::~FCameraShakePreviewer()
 		UnRegisterViewModifier();
 	}
 
+	Teardown();
+}
+
+void FCameraShakePreviewer::Initialize(UWorld* InWorld)
+{
+	FActorSpawnParameters SpawnInfo;
+	SpawnInfo.ObjectFlags |= RF_Transient;
+	SpawnInfo.SpawnCollisionHandlingOverride = ESpawnActorCollisionHandlingMethod::AlwaysSpawn;
+
+	PreviewCamera = InWorld->SpawnActor<APreviewPlayerCameraManager>(SpawnInfo);
+	PreviewCamera->SetIsTemporarilyHiddenInEditor(true);
+
+	PreviewCameraShake = CastChecked<UCameraModifier_CameraShake>(
+		PreviewCamera->AddNewCameraModifier(UCameraModifier_CameraShake::StaticClass()));
+
+	LastDeltaTime = 0.f;
+}
+
+void FCameraShakePreviewer::Teardown()
+{
+	if (PreviewCamera != nullptr)
+	{
+		PreviewCamera->Destroy(false, false);
+	}
+
 	PreviewCameraShake = nullptr;
+	PreviewCamera = nullptr;
 }
 
 void FCameraShakePreviewer::Update(float DeltaTime, bool bIsPlaying)
@@ -37,15 +66,26 @@ void FCameraShakePreviewer::Update(float DeltaTime, bool bIsPlaying)
 	}
 }
 
+void FCameraShakePreviewer::ModifyView(FEditorViewportViewModifierParams& Params)
+{
+	OnModifyView(Params);
+}
+
 void FCameraShakePreviewer::OnModifyView(FEditorViewportViewModifierParams& Params)
 {
 	const float DeltaTime = LastDeltaTime.Get(-1.f);
 	if (DeltaTime > 0.f)
 	{
+		LastPostProcessSettings.Reset();
+		LastPostProcessBlendWeights.Reset();
+		PreviewCamera->ResetPostProcessSettings();
+
 		FMinimalViewInfo OriginalPOV(Params.ViewInfo);
 
 		PreviewCameraShake->ModifyCamera(DeltaTime, Params.ViewInfo);
 
+		PreviewCamera->MergePostProcessSettings(LastPostProcessSettings, LastPostProcessBlendWeights);
+
 		LastLocationModifier = Params.ViewInfo.Location - OriginalPOV.Location;
 		LastRotationModifier = Params.ViewInfo.Rotation - OriginalPOV.Rotation;
 		LastFOVModifier = Params.ViewInfo.FOV - OriginalPOV.FOV;
@@ -58,6 +98,11 @@ void FCameraShakePreviewer::OnModifyView(FEditorViewportViewModifierParams& Para
 		Params.ViewInfo.Rotation += LastRotationModifier;
 		Params.ViewInfo.FOV += LastFOVModifier;
 	}
+
+	for (int32 PPIndex = 0; PPIndex < LastPostProcessSettings.Num(); ++PPIndex)
+	{
+		Params.AddPostProcessBlend(LastPostProcessSettings[PPIndex], LastPostProcessBlendWeights[PPIndex]);
+	}
 }
 
 void FCameraShakePreviewer::RegisterViewModifier()
diff --git a/Engine/Source/Runtime/MovieSceneTracks/Private/Evaluation/MovieSceneCameraShakePreviewer.h b/Engine/Source/Runtime/MovieSceneTracks/Private/Evaluation/MovieSceneCameraShakePreviewer.h
deleted file mode 100644
diff --git a/Engine/Source/Runtime/MovieSceneTracks/Private/Evaluation/MovieSceneCameraShakeSourceShakeTemplate.cpp b/Engine/Source/Runtime/MovieSceneTracks/Private/Evaluation/MovieSceneCameraShakeSourceShakeTemplate.cpp
--- a/Engine/Source/Runtime/MovieSceneTracks/Private/Evaluation/MovieSceneCameraShakeSourceShakeTemplate.cpp
+++ b/Engine/Source/Runtime/MovieSceneTracks/Private/Evaluation/MovieSceneCameraShakeSourceShakeTemplate.cpp
@@ -119,7 +119,7 @@ struct FCameraShakeSourceShakeStartExecutionToken : IMovieSceneExecutionToken
 
 #if WITH_EDITOR
 						// Also start playing the shake in our editor preview.
-						UCameraModifier_CameraShake* const PreviewCameraShake = InstanceData.Previewer.GetCameraShake();
+						UCameraModifier_CameraShake* const PreviewCameraShake = InstanceData.Previewer.GetCameraModifier();
 
 						FAddCameraShakeParams Params;
 						Params.SourceComponent = ShakeSourceComponent;
@@ -176,7 +176,7 @@ struct FCameraShakeSourceShakeBlendOutExecutionToken : IMovieSceneExecutionToken
 						ShakeSourceComponent->StopAllCameraShakesOfType(ShakeClass, bImmediately);
 
 #if WITH_EDITOR
-						UCameraModifier_CameraShake* const PreviewCameraShake = InstanceData.Previewer.GetCameraShake();
+						UCameraModifier_CameraShake* const PreviewCameraShake = InstanceData.Previewer.GetCameraModifier();
 						PreviewCameraShake->RemoveAllCameraShakesOfClassFromSource(ShakeClass, ShakeSourceComponent, bImmediately);
 #endif
 					}
@@ -211,6 +211,10 @@ void FMovieSceneCameraShakeSourceShakeSectionTemplate::Setup(FPersistentEvaluati
 	InstanceData.Status = ECameraShakeSourceShakeStatus::NotStarted;
 
 #if WITH_EDITOR
+	if (!InstanceData.Previewer.IsInitialized())
+	{
+		InstanceData.Previewer.Initialize(Player.GetPlaybackContext()->GetWorld());
+	}
 	InstanceData.Previewer.RegisterViewModifier();
 #endif
 }
diff --git a/Engine/Source/Runtime/MovieSceneTracks/Private/Evaluation/MovieSceneCameraShakeSourceTriggerTemplate.cpp b/Engine/Source/Runtime/MovieSceneTracks/Private/Evaluation/MovieSceneCameraShakeSourceTriggerTemplate.cpp
--- a/Engine/Source/Runtime/MovieSceneTracks/Private/Evaluation/MovieSceneCameraShakeSourceTriggerTemplate.cpp
+++ b/Engine/Source/Runtime/MovieSceneTracks/Private/Evaluation/MovieSceneCameraShakeSourceTriggerTemplate.cpp
@@ -50,7 +50,7 @@ struct FCameraShakeSourceTriggerExecutionToken : IMovieSceneExecutionToken
 
 #if WITH_EDITOR
 							// Also start playing the shake in our editor preview.
-							UCameraModifier_CameraShake* const PreviewCameraShake = InstanceData.Previewer.GetCameraShake();
+							UCameraModifier_CameraShake* const PreviewCameraShake = InstanceData.Previewer.GetCameraModifier();
 
 							FAddCameraShakeParams Params;
 							Params.SourceComponent = ShakeSourceComponent;
@@ -96,6 +96,10 @@ void FMovieSceneCameraShakeSourceTriggerSectionTemplate::Setup(FPersistentEvalua
 {
 	FCameraShakeSourceTriggerInstanceData& InstanceData = PersistentData.AddSectionData<FCameraShakeSourceTriggerInstanceData>();
 #if WITH_EDITOR
+	if (!InstanceData.Previewer.IsInitialized())
+	{
+		InstanceData.Previewer.Initialize(Player.GetPlaybackContext()->GetWorld());
+	}
 	InstanceData.Previewer.RegisterViewModifier();
 #endif
 }
diff --git a/Engine/Source/Runtime/MovieSceneTracks/Public/Evaluation/MovieSceneCameraShakePreviewer.h b/Engine/Source/Runtime/MovieSceneTracks/Public/Evaluation/MovieSceneCameraShakePreviewer.h
new file mode 100644
--- /dev/null
+++ b/Engine/Source/Runtime/MovieSceneTracks/Public/Evaluation/MovieSceneCameraShakePreviewer.h
@@ -0,0 +1,84 @@
+// Copyright Epic Games, Inc. All Rights Reserved.
+
+#pragma once
+
+#if WITH_EDITOR
+
+#include "CoreMinimal.h"
+#include "UObject/GCObject.h"
+#include "Camera/PlayerCameraManager.h"
+
+#include "MovieSceneCameraShakePreviewer.generated.h"
+
+class FLevelEditorViewportClient;
+class UCameraModifier_CameraShake;
+struct FEditorViewportViewModifierParams;
+
+/**
+ * A mock player camera manager, used to store camera shake post-process settings for previewing.
+ */
+UCLASS()
+class APreviewPlayerCameraManager : public APlayerCameraManager
+{
+	GENERATED_BODY()
+
+public:
+	void ResetPostProcessSettings()
+	{
+		ClearCachedPPBlends();
+	}
+
+	void MergePostProcessSettings(TArray<FPostProcessSettings>& InSettings, TArray<float>& InBlendWeights)
+	{
+		InSettings.Append(PostProcessBlendCache);
+		InBlendWeights.Append(PostProcessBlendCacheWeights);
+	}
+};
+
+/**
+ * A class that owns a gameplay camera shake manager, so that we can us it to preview shakes in editor.
+ */
+class MOVIESCENETRACKS_API FCameraShakePreviewer : public FGCObject
+{
+public:
+	FCameraShakePreviewer();
+	~FCameraShakePreviewer();
+
+	void Initialize(UWorld* InWorld);
+	bool IsInitialized() const { return PreviewCameraShake != nullptr; }
+	void Teardown();
+
+	APlayerCameraManager* GetCameraManager() const { return PreviewCamera; }
+	UCameraModifier_CameraShake* GetCameraModifier() const { return PreviewCameraShake; }
+
+	void ModifyView(FEditorViewportViewModifierParams& Params);
+
+	void RegisterViewModifier();
+	void UnRegisterViewModifier();
+	void Update(float DeltaTime, bool bIsPlaying);
+
+private:
+	// FGCObject interface
+	virtual void AddReferencedObjects(FReferenceCollector& Collector) override { Collector.AddReferencedObject(PreviewCameraShake); }
+	virtual FString GetReferencerName() const override { return TEXT("SCameraShakePreviewer"); }
+
+private:
+	void OnModifyView(FEditorViewportViewModifierParams& Params);
+	void OnLevelViewportClientListChanged();
+
+private:
+	APreviewPlayerCameraManager* PreviewCamera;
+	UCameraModifier_CameraShake* PreviewCameraShake;
+	TArray<FLevelEditorViewportClient*> RegisteredViewportClients;
+
+	TOptional<float> LastDeltaTime;
+	FVector LastLocationModifier;
+	FRotator LastRotationModifier;
+	float LastFOVModifier;
+
+	TArray<FPostProcessSettings> LastPostProcessSettings;
+	TArray<float> LastPostProcessBlendWeights;
+};
+
+#endif
+
