Sequencer: make the camera shake previewer aware of the world in which it should run shakes

Making it aware of which world to operate on makes it possible to create the previewed camera shakes outered to that world, which prevents an assert when the shake sequence player tries to find the subsystem to evaluate sequences.

The change however forces us to make the camera shake previewer tool also aware of the current world, and better manage whatever viewports it's hooked up to.

diff --git a/Engine/Plugins/Cameras/CameraShakePreviewer/Source/CameraShakePreviewer/Private/SCameraShakePreviewer.cpp b/Engine/Plugins/Cameras/CameraShakePreviewer/Source/CameraShakePreviewer/Private/SCameraShakePreviewer.cpp
--- a/Engine/Plugins/Cameras/CameraShakePreviewer/Source/CameraShakePreviewer/Private/SCameraShakePreviewer.cpp
+++ b/Engine/Plugins/Cameras/CameraShakePreviewer/Source/CameraShakePreviewer/Private/SCameraShakePreviewer.cpp
@@ -17,7 +17,7 @@
 
 #define LOCTEXT_NAMESPACE "CameraShakePreviewer"
 
-namespace UE::MovieScene
+namespace UE::Sequencer
 {
 
 UWorld* FindCameraShakePreviewerWorld()
@@ -58,21 +58,14 @@ struct FCameraShakeData
 	}
 };
 
-FCameraShakePreviewUpdater::FCameraShakePreviewUpdater()
+FCameraShakePreviewUpdater::FCameraShakePreviewUpdater(UWorld* InWorld)
+	: Previewer(InWorld)
 {
-	// Handle camera shakes being recompiled.
-	FCoreUObjectDelegates::OnObjectsReplaced.AddRaw(this, &FCameraShakePreviewUpdater::OnObjectsReplaced);
 }
 
 FCameraShakePreviewUpdater::~FCameraShakePreviewUpdater()
 {
-	FEditorSupportDelegates::PrepareToCleanseEditorObject.RemoveAll(this);
-	FCoreUObjectDelegates::OnObjectsReplaced.RemoveAll(this);
-}
-
-void FCameraShakePreviewUpdater::OnObjectsReplaced(const TMap<UObject*, UObject*>& ReplacementMap)
-{
-	Previewer.OnObjectsReplaced(ReplacementMap);
+	Previewer.UnRegisterViewModifiers();
 }
 
 void FCameraShakePreviewUpdater::Tick(float DeltaTime)
@@ -80,11 +73,6 @@ void FCameraShakePreviewUpdater::Tick(float DeltaTime)
 	Previewer.Update(DeltaTime, true);
 }
 
-void FCameraShakePreviewUpdater::ModifyCamera(FEditorViewportViewModifierParams& Params)
-{
-	Previewer.ModifyView(Params);
-}
-
 UCameraShakeBase* FCameraShakePreviewUpdater::AddCameraShake(TSubclassOf<UCameraShakeBase> ShakeClass, const FAddCameraShakeParams& Params)
 {
 	FCameraShakePreviewerAddParams ActualParams;
@@ -314,11 +302,8 @@ void SCameraShakePreviewer::Construct(const FArguments& InArgs)
 		GEditor->OnLevelViewportClientListChanged().AddSP(this, &SCameraShakePreviewer::OnLevelViewportClientListChanged);
 	}
 
-	// Create our camera shake manager.
-	CameraShakePreviewUpdater = TUniquePtr<FCameraShakePreviewUpdater>(new FCameraShakePreviewUpdater());
-	ActiveViewportClient = nullptr;
-
 	// Populate the main list based on the current level.
+	UpdateActiveViewportAndWorld();
 	Populate();
 	bNeedsRefresh = false;
 }
@@ -327,11 +312,6 @@ SCameraShakePreviewer::~SCameraShakePreviewer()
 {
 	CameraShakePreviewUpdater = nullptr;
 
-	if (ActiveViewportClient != nullptr)
-	{
-		ActiveViewportClient->ViewModifiers.RemoveAll(this);
-	}
-
 	FEditorDelegates::MapChange.RemoveAll(this);
 	FEditorDelegates::NewCurrentLevel.RemoveAll(this);
 	
@@ -361,7 +341,7 @@ TSharedRef<ITableRow> SCameraShakePreviewer::OnCameraShakesListGenerateRowWidget
 
 void SCameraShakePreviewer::Populate()
 {
-	UWorld* CurrentWorld = FindCurrentWorld();
+	UWorld* CurrentWorld =  WeakCurrentWorld.Get();
 	if (!ensureMsgf(CurrentWorld, TEXT("Could not find current world instance.")))
 	{
 		return;
@@ -437,12 +417,6 @@ void SCameraShakePreviewer::Populate()
 	CameraShakesListView->RequestListRefresh();
 }
 
-UWorld* SCameraShakePreviewer::FindCurrentWorld()
-{
-	WeakCurrentWorld = UE::MovieScene::FindCameraShakePreviewerWorld();
-	return WeakCurrentWorld.Get();
-}
-
 void SCameraShakePreviewer::Refresh()
 {
 	bNeedsRefresh = true;
@@ -452,6 +426,9 @@ void SCameraShakePreviewer::Tick(const FGeometry& AllottedGeometry, const double
 {
 	SCompoundWidget::Tick(AllottedGeometry, InCurrentTime, InDeltaTime);
 
+	// See if the world or active viewport have changed.
+	UpdateActiveViewportAndWorld();
+
 	// Update our list of camera shakes if needed.
 	if (bNeedsRefresh)
 	{
@@ -518,66 +495,73 @@ void SCameraShakePreviewer::Tick(const FGeometry& AllottedGeometry, const double
 			}
 		}
 	}
+}
 
-	// See if the active viewport has changed.
-	if (GEditor)
+void SCameraShakePreviewer::UpdateActiveViewportAndWorld()
+{
+	if (!GEditor)
 	{
-		FViewport* ActiveViewport = GEditor->GetActiveViewport();
-		if ((ActiveViewportClient == nullptr && ActiveViewport != nullptr) ||
-			(ActiveViewportClient != nullptr && ActiveViewportClient->Viewport != ActiveViewport))
-		{
-			ActiveViewportIndex = 0;
-			FLevelEditorViewportClient* NewActiveViewportClient = nullptr;
-			if (ActiveViewport != nullptr)
-			{
-				// Find the new active viewport client.
-				const TArray<FLevelEditorViewportClient*>& ViewportClients = GEditor->GetLevelViewportClients();
-				for (FLevelEditorViewportClient* ViewportClient : ViewportClients)
-				{
-					++ActiveViewportIndex;
-					if (ViewportClient->Viewport == ActiveViewport)
-					{
-						NewActiveViewportClient = ViewportClient;
-						break;
-					}
-				}
-			}
-			// else: no new active viewport client, and leave the index at 0.
+		ActiveViewportIndex = INDEX_NONE;
+		WeakCurrentWorld = nullptr;
+		CameraShakePreviewUpdater.Reset();
+		return;
+	}
 
-			if (NewActiveViewportClient != ActiveViewportClient)
-			{
-				// Clear the old viewport's callbacks.
-				if (ActiveViewportClient != nullptr)
-				{
-					ActiveViewportClient->ViewModifiers.RemoveAll(this);
-				}
+	// Update the active viewport and index. This is used for showing viewport info in the UI.
+	ActiveViewportIndex = 0;
+	ActiveViewportClient = nullptr;
+	FViewport* const ActiveViewport = GEditor->GetActiveViewport();
+	for (FLevelEditorViewportClient* ViewportClient : GEditor->GetLevelViewportClients())
+	{
+		if (ViewportClient->Viewport == ActiveViewport)
+		{
+			ActiveViewportClient = ViewportClient;
+			break;
+		}
+		++ActiveViewportIndex;
+	}
+	if (!ActiveViewportClient)
+	{
+		ActiveViewportIndex = INDEX_NONE;
+	}
 
-				ActiveViewportClient = NewActiveViewportClient;
+	// Always use the first editor world we can find. If this is a different world than what we
+	// were looking at before, re-create the shake previewer.
+	UWorld* const PreviousWorld = WeakCurrentWorld.Get();
+	UWorld* const NewWorld = UE::Sequencer::FindCameraShakePreviewerWorld();
+	if (PreviousWorld != NewWorld)
+	{
+		if (NewWorld)
+		{
+			CameraShakePreviewUpdater.Reset(new FCameraShakePreviewUpdater(NewWorld));
 
-				// Add to the new viewport's callbacks if necessary.
-				if (ActiveViewportClient != nullptr)
-				{
-					if (ActiveViewportClient != nullptr && CameraShakePreviewerModule->HasCameraShakesPreview(ActiveViewportClient))
+			// Immediately register the previewer with any viewport that wants shaking.
+			CameraShakePreviewUpdater->GetPreviewer().RegisterViewModifiers(
+					[this, NewWorld](FLevelEditorViewportClient* LevelVC)
 					{
-						ActiveViewportClient->ViewModifiers.AddRaw(this, &SCameraShakePreviewer::OnModifyView);
-					}
-				}
-			}
+						return LevelVC->GetWorld() == NewWorld && 
+							CameraShakePreviewerModule->HasCameraShakesPreview(LevelVC);
+					});
+		}
+		else
+		{
+			CameraShakePreviewUpdater.Reset();
 		}
 	}
+	WeakCurrentWorld = NewWorld;
 }
 
 void SCameraShakePreviewer::OnTogglePreviewCameraShakes(const FTogglePreviewCameraShakesParams& Params)
 {
-	if (Params.ViewportClient == ActiveViewportClient)
+	if (Params.ViewportClient && Params.ViewportClient->GetWorld() == WeakCurrentWorld)
 	{
-		if (Params.bPreviewCameraShakes && ensure(!Params.ViewportClient->ViewModifiers.IsBoundToObject(this)))
+		if (Params.bPreviewCameraShakes)
 		{
-			ActiveViewportClient->ViewModifiers.AddRaw(this, &SCameraShakePreviewer::OnModifyView);
+			CameraShakePreviewUpdater->GetPreviewer().RegisterViewModifier(Params.ViewportClient);
 		}
-		else if (!Params.bPreviewCameraShakes && ensure(Params.ViewportClient->ViewModifiers.IsBoundToObject(this)))
+		else if (!Params.bPreviewCameraShakes)
 		{
-			ActiveViewportClient->ViewModifiers.RemoveAll(this);
+			CameraShakePreviewUpdater->GetPreviewer().UnRegisterViewModifier(Params.ViewportClient);
 		}
 	}
 }
@@ -596,9 +580,9 @@ void SCameraShakePreviewer::OnCameraShakesListSectionChanged(TSharedPtr<FCameraS
 
 FText SCameraShakePreviewer::GetActiveViewportName() const
 {
-	if (ActiveViewportIndex > 0)
+	if (ActiveViewportIndex != INDEX_NONE)
 	{
-		return FText::FromString(LexToString(ActiveViewportIndex));
+		return FText::FromString(LexToString(ActiveViewportIndex + 1));
 	}
 	return FText::FromString("<None>");
 }
@@ -690,17 +674,18 @@ void SCameraShakePreviewer::PlayCameraShake(TSharedPtr<FCameraShakeData> CameraS
 
 void SCameraShakePreviewer::OnLevelViewportClientListChanged()
 {
-	if (ActiveViewportClient != nullptr)
+	// When viewports change, our shake previewer already correctly unregisters from any removed viewport.
+	// However, we need to automatically register any *new* viewport that fits our requirements.
+	if (UWorld* CurrentWorld = WeakCurrentWorld.Get())
 	{
-		const TArray<FLevelEditorViewportClient*> ViewportClients = GEditor->GetLevelViewportClients();
-		const bool bFound = ViewportClients.Contains(ActiveViewportClient);
-		if (!bFound)
-		{
-			// Our target viewport has been removed from the list... let's forget it.
-			// TODO: unsafe? maybe we don't actually need to cleanup?
-			ActiveViewportClient->ViewModifiers.RemoveAll(this);
-			ActiveViewportClient = nullptr;
-		}
+		CameraShakePreviewUpdater->GetPreviewer().RegisterViewModifiers(
+				[this, CurrentWorld](FLevelEditorViewportClient* LevelVC)
+				{
+					return LevelVC->GetWorld() == CurrentWorld && 
+						CameraShakePreviewerModule->HasCameraShakesPreview(LevelVC);
+				},
+				// Ignore duplicate registrations.
+				true);
 	}
 }
 
@@ -749,9 +734,5 @@ void SCameraShakePreviewer::PostUndo(bool bSuccess)
 	Refresh();
 }
 
-void SCameraShakePreviewer::OnModifyView(FEditorViewportViewModifierParams& Params)
-{
-	CameraShakePreviewUpdater->ModifyCamera(Params);
-}
-
 #undef LOCTEXT_NAMESPACE
+
diff --git a/Engine/Plugins/Cameras/CameraShakePreviewer/Source/CameraShakePreviewer/Private/SCameraShakePreviewer.h b/Engine/Plugins/Cameras/CameraShakePreviewer/Source/CameraShakePreviewer/Private/SCameraShakePreviewer.h
--- a/Engine/Plugins/Cameras/CameraShakePreviewer/Source/CameraShakePreviewer/Private/SCameraShakePreviewer.h
+++ b/Engine/Plugins/Cameras/CameraShakePreviewer/Source/CameraShakePreviewer/Private/SCameraShakePreviewer.h
@@ -40,7 +40,7 @@ struct FTogglePreviewCameraShakesParams;
 class FCameraShakePreviewUpdater : public FTickableEditorObject
 {
 public:
-	FCameraShakePreviewUpdater();
+	FCameraShakePreviewUpdater(UWorld* InWorld);
 	virtual ~FCameraShakePreviewUpdater();
 
 	// FTickableObject Interface
@@ -54,10 +54,7 @@ public:
 	void RemoveAllCameraShakesFromSource(const UCameraShakeSourceComponent* SourceComponent);
 	void RemoveAllCameraShakes();
 
-	void ModifyCamera(FEditorViewportViewModifierParams& Params);
-
-private:
-	void OnObjectsReplaced(const TMap<UObject*, UObject*>& ReplacementMap);
+	FCameraShakePreviewer& GetPreviewer() { return Previewer; }
 
 private:
 	FCameraShakePreviewer Previewer;
@@ -84,8 +81,8 @@ public:
 
 private:
 	void Populate();
-	UWorld* FindCurrentWorld();
 	void Refresh();
+	void UpdateActiveViewportAndWorld();
 
 	void OnTogglePreviewCameraShakes(const FTogglePreviewCameraShakesParams& Params);
 
@@ -108,8 +105,6 @@ private:
 	void OnNewCurrentLevel();
 	void OnMapLoaded(const FString&  Filename, bool bAsTemplate);
 
-	void OnModifyView(FEditorViewportViewModifierParams& Params);
-
 private:
 	TArray<TSharedPtr<FCameraShakeData>> CameraShakes;
 	TUniquePtr<FCameraShakePreviewUpdater> CameraShakePreviewUpdater;
@@ -118,9 +113,11 @@ private:
 	TSharedPtr<SButton> PlayStopSelectedButton;
 
 	FCameraShakePreviewerModule* CameraShakePreviewerModule;
+
 	FLevelEditorViewportClient* ActiveViewportClient;
-	int ActiveViewportIndex;
+	int32 ActiveViewportIndex;
 
 	TWeakObjectPtr<UWorld> WeakCurrentWorld;
+
 	bool bNeedsRefresh;
 };
diff --git a/Engine/Source/Runtime/MovieSceneTracks/Private/Evaluation/MovieSceneCameraShakePreviewer.cpp b/Engine/Source/Runtime/MovieSceneTracks/Private/Evaluation/MovieSceneCameraShakePreviewer.cpp
--- a/Engine/Source/Runtime/MovieSceneTracks/Private/Evaluation/MovieSceneCameraShakePreviewer.cpp
+++ b/Engine/Source/Runtime/MovieSceneTracks/Private/Evaluation/MovieSceneCameraShakePreviewer.cpp
@@ -9,25 +9,30 @@
 #include "Camera/CameraShakeSourceComponent.h"
 #include "Camera/CameraModifier_CameraShake.h"
 
-FCameraShakePreviewer::FCameraShakePreviewer()
-	: LastDeltaTime(0.f)
+FCameraShakePreviewer::FCameraShakePreviewer(UWorld* InWorld)
+	: World(InWorld)
+	, LastDeltaTime(0.f)
 	, LastLocationModifier(FVector::ZeroVector)
 	, LastRotationModifier(FRotator::ZeroRotator)
 	, LastFOVModifier(0.f)
 {
+	// Handle camera shakes being recompiled.
+	FCoreUObjectDelegates::OnObjectsReplaced.AddRaw(this, &FCameraShakePreviewer::OnObjectsReplaced);
 }
 
 FCameraShakePreviewer::~FCameraShakePreviewer()
 {
-	if (!ensureMsgf(RegisteredViewportClients.Num() == 0, TEXT("Forgot to call UnRegisterViewModifier!")))
+	if (!ensureMsgf(RegisteredViewportClients.Num() == 0, TEXT("Forgot to call UnRegisterViewModifiers!")))
 	{
-		UnRegisterViewModifier();
+		UnRegisterViewModifiers();
 	}
+
+	FCoreUObjectDelegates::OnObjectsReplaced.RemoveAll(this);
 }
 
 UCameraShakeBase* FCameraShakePreviewer::AddCameraShake(const FCameraShakePreviewerAddParams& Params)
 {
-	UCameraShakeBase* NewShake = NewObject<UCameraShakeBase>(GetTransientPackage(), Params.ShakeClass);
+	UCameraShakeBase* NewShake = NewObject<UCameraShakeBase>(World, Params.ShakeClass);
 	ActiveShakes.Add({ NewShake, Params.SourceComponent, Params.GlobalStartTime });
 
 	FCameraShakeBaseStartParams StartParams;
@@ -211,7 +216,12 @@ void FCameraShakePreviewer::OnModifyView(FEditorViewportViewModifierParams& Para
 	}
 }
 
-void FCameraShakePreviewer::RegisterViewModifier()
+void FCameraShakePreviewer::RegisterViewModifiers(bool bIgnoreDuplicateRegistration)
+{
+	RegisterViewModifiers([](FLevelEditorViewportClient*) { return true; }, bIgnoreDuplicateRegistration);
+}
+
+void FCameraShakePreviewer::RegisterViewModifiers(FViewportFilter ViewportFilter, bool bIgnoreDuplicateRegistration)
 {
 	if (GEditor == nullptr)
 	{
@@ -221,20 +231,38 @@ void FCameraShakePreviewer::RegisterViewModifier()
 	// Register our view modifier on all appropriate viewports, and remember which viewports we did that on.
 	// We will later make sure to unregister on the same list, except for any viewport that somehow disappeared since,
 	// which we will be notified about with the OnLevelViewportClientListChanged event.
-	RegisteredViewportClients.Reset();
 	for (FLevelEditorViewportClient* LevelVC : GEditor->GetLevelViewportClients())
 	{		
-		if (LevelVC && LevelVC->AllowsCinematicControl() && LevelVC->GetViewMode() != VMI_Unknown)
+		if (LevelVC && 
+				LevelVC->GetViewMode() != VMI_Unknown &&
+				LevelVC->GetWorld() == World &&
+				ViewportFilter(LevelVC))
 		{
-			RegisteredViewportClients.Add(LevelVC);
-			LevelVC->ViewModifiers.AddRaw(this, &FCameraShakePreviewer::OnModifyView);
+			RegisterViewModifier(LevelVC, bIgnoreDuplicateRegistration);
 		}
 	}
+}
 
-	GEditor->OnLevelViewportClientListChanged().AddRaw(this, &FCameraShakePreviewer::OnLevelViewportClientListChanged);
+void FCameraShakePreviewer::RegisterViewModifier(FLevelEditorViewportClient* ViewportClient, bool bIgnoreDuplicateRegistration)
+{
+	if (RegisteredViewportClients.Contains(ViewportClient))
+	{
+		// Already registered to this viewport.
+		ensureMsgf(bIgnoreDuplicateRegistration, TEXT("Given viewport is already registered."));
+		return;
+	}
+
+	RegisteredViewportClients.Add(ViewportClient);
+	ViewportClient->ViewModifiers.AddRaw(this, &FCameraShakePreviewer::OnModifyView);
+	if (RegisteredViewportClients.Num() == 1)
+	{
+		// If this is the first viewport, start listening to viewports changing.
+		GEditor->OnLevelViewportClientListChanged().AddRaw(this, &FCameraShakePreviewer::OnLevelViewportClientListChanged);
+	}
 }
 
-void FCameraShakePreviewer::UnRegisterViewModifier()
+
+void FCameraShakePreviewer::UnRegisterViewModifiers()
 {
 	if (GEditor == nullptr)
 	{
@@ -250,6 +278,21 @@ void FCameraShakePreviewer::UnRegisterViewModifier()
 	RegisteredViewportClients.Reset();
 }
 
+void FCameraShakePreviewer::UnRegisterViewModifier(FLevelEditorViewportClient* ViewportClient)
+{
+	const int32 NumRemoved = RegisteredViewportClients.Remove(ViewportClient);
+	if (ensureMsgf(NumRemoved > 0, TEXT("The given viewport client wasn't registered.")))
+	{
+		// If this is the last viewport, stop listening to viewports changing.
+		if (RegisteredViewportClients.IsEmpty())
+		{
+			GEditor->OnLevelViewportClientListChanged().RemoveAll(this);
+		}
+
+		ViewportClient->ViewModifiers.RemoveAll(this);
+	}
+}
+
 void FCameraShakePreviewer::OnLevelViewportClientListChanged()
 {
 	if (GEditor != nullptr)
@@ -259,6 +302,21 @@ void FCameraShakePreviewer::OnLevelViewportClientListChanged()
 		TSet<FLevelEditorViewportClient*> PreviousViewportClients(RegisteredViewportClients);
 		TSet<FLevelEditorViewportClient*> NewViewportClients(GEditor->GetLevelViewportClients());
 		RegisteredViewportClients = PreviousViewportClients.Intersect(NewViewportClients).Array();
+
+		// Unregister ourselves from the viewport-change callback if we don't have any registered 
+		// viewports anymore.
+		if (RegisteredViewportClients.IsEmpty())
+		{
+			GEditor->OnLevelViewportClientListChanged().RemoveAll(this);
+		}
+
+		// Remove ourselves from the viewport that went away. It's probably not necessary since they
+		// will be destroyed, but better be safe.
+		TSet<FLevelEditorViewportClient*> RemovedViewportClients = PreviousViewportClients.Difference(NewViewportClients);
+		for (FLevelEditorViewportClient* RemovedViewportClient : RemovedViewportClients)
+		{
+			RemovedViewportClient->ViewModifiers.RemoveAll(this);
+		}
 	}
 }
 
diff --git a/Engine/Source/Runtime/MovieSceneTracks/Private/Systems/MovieSceneCameraShakeSystem.cpp b/Engine/Source/Runtime/MovieSceneTracks/Private/Systems/MovieSceneCameraShakeSystem.cpp
--- a/Engine/Source/Runtime/MovieSceneTracks/Private/Systems/MovieSceneCameraShakeSystem.cpp
+++ b/Engine/Source/Runtime/MovieSceneTracks/Private/Systems/MovieSceneCameraShakeSystem.cpp
@@ -20,6 +20,7 @@
 
 #if WITH_EDITOR
 #include "Editor.h"
+#include "LevelEditorViewport.h"
 #endif  // WITH_EDITOR
 
 #include UE_INLINE_GENERATED_CPP_BY_NAME(MovieSceneCameraShakeSystem)
@@ -150,10 +151,12 @@ FCameraShakePreviewerLinkerExtension::FCameraShakePreviewerLinkerExtension(UMovi
 
 FCameraShakePreviewerLinkerExtension::~FCameraShakePreviewerLinkerExtension()
 {
+	GEditor->OnLevelViewportClientListChanged().RemoveAll(this);
+
 	for (TPair<FInstanceHandle, FCameraShakePreviewer>& Pair : Previewers)
 	{
 		FCameraShakePreviewer& Previewer = Pair.Value;
-		Previewer.UnRegisterViewModifier();
+		Previewer.UnRegisterViewModifiers();
 	}
 	Previewers.Reset();
 }
@@ -170,13 +173,30 @@ FCameraShakePreviewer& FCameraShakePreviewerLinkerExtension::GetPreviewer(FInsta
 		return *Previewer;
 	}
 
-	FCameraShakePreviewer& NewPreviewer = Previewers.Add(InstanceHandle);
-	NewPreviewer.RegisterViewModifier();
+	if (Previewers.IsEmpty())
+	{
+		// This is our first previewer... let's start listening to viewports changing.
+		GEditor->OnLevelViewportClientListChanged().AddSP(this, &FCameraShakePreviewerLinkerExtension::OnLevelViewportClientListChanged);
+	}
+
+	UMovieSceneEntitySystemLinker* Linker = WeakLinker.Get();
+	FInstanceRegistry* InstanceRegistry = Linker->GetInstanceRegistry();
+	const FSequenceInstance& SequenceInstance = InstanceRegistry->GetInstance(InstanceHandle);
+	UObject* PlaybackContext = SequenceInstance.GetPlayer()->GetPlaybackContext();
+	UWorld* ContextWorld = PlaybackContext ? PlaybackContext->GetWorld() : nullptr;
+
+	FCameraShakePreviewer& NewPreviewer = Previewers.Add(InstanceHandle, ContextWorld);
+	NewPreviewer.RegisterViewModifiers([ContextWorld](FLevelEditorViewportClient* LevelVC) -> bool
+				{
+					return LevelVC->AllowsCinematicControl() && LevelVC->GetWorld() == ContextWorld;
+				});
 	return NewPreviewer;
 }
 
-void FCameraShakePreviewerLinkerExtension::UpdateAllPreviewers(FInstanceRegistry* InstanceRegistry)
+void FCameraShakePreviewerLinkerExtension::UpdateAllPreviewers()
 {
+	UMovieSceneEntitySystemLinker* Linker = WeakLinker.Get();
+	FInstanceRegistry* InstanceRegistry = Linker->GetInstanceRegistry();
 	const TSparseArray<FSequenceInstance>& Instances = InstanceRegistry->GetSparseInstances();
 	for (auto It = Instances.CreateConstIterator(); It; ++It)
 	{
@@ -212,6 +232,28 @@ bool FCameraShakePreviewerLinkerExtension::HasAnyShake() const
 	return false;
 }
 
+void FCameraShakePreviewerLinkerExtension::OnLevelViewportClientListChanged()
+{
+	// When viewports change, our shake previewers already correctly unregister from any removed viewport.
+	// However, we need to automatically register any *new* viewport that fits our requirements.
+	UMovieSceneEntitySystemLinker* Linker = WeakLinker.Get();
+	FInstanceRegistry* InstanceRegistry = Linker->GetInstanceRegistry();
+	for (TPair<FInstanceHandle, FCameraShakePreviewer> Pair : Previewers)
+	{
+		const FSequenceInstance& SequenceInstance = InstanceRegistry->GetInstance(Pair.Key);
+		UObject* PlaybackContext = SequenceInstance.GetPlayer()->GetPlaybackContext();
+		UWorld* ContextWorld = PlaybackContext ? PlaybackContext->GetWorld() : nullptr;
+
+		Pair.Value.RegisterViewModifiers(
+				[ContextWorld](FLevelEditorViewportClient* LevelVC)
+				{
+					return LevelVC->AllowsCinematicControl() && LevelVC->GetWorld() == ContextWorld;
+				},
+				// Ignore duplicate registrations.
+				true);
+	}
+}
+
 #endif  // WITH_EDITOR
 
 } // namespace UE::MovieScene
@@ -697,8 +739,7 @@ void UMovieSceneCameraShakeEvaluatorSystem::OnRun(FSystemTaskPrerequisites& InPr
 		// The previewer only stores the delta time, and only computes shake results when the editor
 		// later processes viewports. We can therefore safely do this in parallel with the shake
 		// evaluation above.
-		FInstanceRegistry* InstanceRegistry = Linker->GetInstanceRegistry();
-		PreviewerExtension->UpdateAllPreviewers(InstanceRegistry);
+		PreviewerExtension->UpdateAllPreviewers();
 	}
 #endif  // WITH_EDITOR
 }
diff --git a/Engine/Source/Runtime/MovieSceneTracks/Public/Evaluation/MovieSceneCameraShakePreviewer.h b/Engine/Source/Runtime/MovieSceneTracks/Public/Evaluation/MovieSceneCameraShakePreviewer.h
--- a/Engine/Source/Runtime/MovieSceneTracks/Public/Evaluation/MovieSceneCameraShakePreviewer.h
+++ b/Engine/Source/Runtime/MovieSceneTracks/Public/Evaluation/MovieSceneCameraShakePreviewer.h
@@ -38,13 +38,21 @@ struct FCameraShakePreviewerAddParams
 class MOVIESCENETRACKS_API FCameraShakePreviewer : public FGCObject
 {
 public:
-	FCameraShakePreviewer();
+	using FViewportFilter = TFunctionRef<bool(FLevelEditorViewportClient*)>;
+
+	FCameraShakePreviewer(UWorld* InWorld);
 	~FCameraShakePreviewer();
 
+	UWorld* GetWorld() const { return World; }
+
 	void ModifyView(FEditorViewportViewModifierParams& Params);
 
-	void RegisterViewModifier();
-	void UnRegisterViewModifier();
+	void RegisterViewModifiers(bool bIgnoreDuplicateRegistration = false);
+	void RegisterViewModifiers(FViewportFilter InViewportFilter, bool bIgnoreDuplicateRegistration = false);
+	void UnRegisterViewModifiers();
+
+	void RegisterViewModifier(FLevelEditorViewportClient* ViewportClient, bool bIgnoreDuplicateRegistration = false);
+	void UnRegisterViewModifier(FLevelEditorViewportClient* ViewportClient);
 
 	void Update(float DeltaTime, bool bIsPlaying);
 	void Scrub(float ScrubTime);
@@ -71,6 +79,8 @@ private:
 	void ResetModifiers();
 
 private:
+	UWorld* World;
+
 	TArray<FLevelEditorViewportClient*> RegisteredViewportClients;
 
 	struct FPreviewCameraShakeInfo
diff --git a/Engine/Source/Runtime/MovieSceneTracks/Public/Systems/MovieSceneCameraShakeSystem.h b/Engine/Source/Runtime/MovieSceneTracks/Public/Systems/MovieSceneCameraShakeSystem.h
--- a/Engine/Source/Runtime/MovieSceneTracks/Public/Systems/MovieSceneCameraShakeSystem.h
+++ b/Engine/Source/Runtime/MovieSceneTracks/Public/Systems/MovieSceneCameraShakeSystem.h
@@ -41,10 +41,13 @@ namespace UE::MovieScene
 		/** Get or create the previewer for the given sequence instance */
 		FCameraShakePreviewer& GetPreviewer(FInstanceHandle InstanceHandle);
 		/** Update all previewers by using the update context of each matching sequence instance */
-		void UpdateAllPreviewers(FInstanceRegistry* InstanceRegistry);
+		void UpdateAllPreviewers();
 		/** Whether there is any previewer with any shake */
 		bool HasAnyShake() const;
 
+	private:
+		void OnLevelViewportClientListChanged();
+
 	private:
 		TMap<FInstanceHandle, FCameraShakePreviewer> Previewers;
 	};
