CIS fix for mac compiler

diff --git a/Engine/Plugins/Cameras/CameraShakePreviewer/Source/CameraShakePreviewer/Private/SCameraShakePreviewer.cpp b/Engine/Plugins/Cameras/CameraShakePreviewer/Source/CameraShakePreviewer/Private/SCameraShakePreviewer.cpp
--- a/Engine/Plugins/Cameras/CameraShakePreviewer/Source/CameraShakePreviewer/Private/SCameraShakePreviewer.cpp
+++ b/Engine/Plugins/Cameras/CameraShakePreviewer/Source/CameraShakePreviewer/Private/SCameraShakePreviewer.cpp
@@ -541,7 +541,9 @@ void SCameraShakePreviewer::UpdateActiveViewportAndWorld()
 					{
 						return LevelVC->GetWorld() == NewWorld && 
 							CameraShakePreviewerModule->HasCameraShakesPreview(LevelVC);
-					});
+					},
+					// Pass false to help the Mac compiler along.
+					false);
 		}
 		else
 		{
