[Viewport Toolbar] Refactor Camera Shake Preview menu extension point logic

Refactor so that it is possible to use Camera Shake Preview toggle via extensions (code used to return empty Extension).

diff --git a/Engine/Plugins/Cameras/CameraShakePreviewer/Source/CameraShakePreviewer/Private/CameraShakePreviewerModule.cpp b/Engine/Plugins/Cameras/CameraShakePreviewer/Source/CameraShakePreviewer/Private/CameraShakePreviewerModule.cpp
--- a/Engine/Plugins/Cameras/CameraShakePreviewer/Source/CameraShakePreviewer/Private/CameraShakePreviewerModule.cpp
+++ b/Engine/Plugins/Cameras/CameraShakePreviewer/Source/CameraShakePreviewer/Private/CameraShakePreviewerModule.cpp
@@ -12,15 +12,39 @@
 #include "WorkspaceMenuStructure.h"
 #include "WorkspaceMenuStructureModule.h"
 #include "SLevelViewport.h"
+#include "Framework/MultiBox/MultiBoxBuilder.h"
 #include "LevelEditorMenuContext.h"
 
 #define LOCTEXT_NAMESPACE "CameraShakePreviewer"
 
 IMPLEMENT_MODULE(FCameraShakePreviewerModule, CameraShakePreviewer);
 
+namespace UE::CameraShakePreviewer::Private
+{
 static const FName LevelEditorModuleName("LevelEditor");
 static const FName LevelEditorCameraShakePreviewerTab("CameraShakePreviewer");
 
+FLevelEditorViewportClient* GetPerspectiveLevelEditorViewportClient()
+{
+	FLevelEditorModule& LevelEditorModule = FModuleManager::LoadModuleChecked<FLevelEditorModule>(LevelEditorModuleName);
+
+	if (TSharedPtr<ILevelEditor> LevelEditor = LevelEditorModule.GetLevelEditorInstance().Pin())
+	{
+		if (TSharedPtr<SLevelViewport> ViewportInterface = LevelEditor->GetActiveViewportInterface())
+		{
+			FLevelEditorViewportClient* ViewportClient =
+				static_cast<FLevelEditorViewportClient*>(&ViewportInterface->GetAssetViewportClient());
+			if (ViewportClient->ViewportType == ELevelViewportType::LVT_Perspective)
+			{
+				return ViewportClient;
+			}
+		}
+	}
+
+	return nullptr;
+}
+} // namespace UE::CameraShakePreviewer::Private
+
 /**
  * Editor commands for the camera shake preview tool.
  */
@@ -54,7 +78,7 @@ void FCameraShakePreviewerModule::StartupModule()
 {
 	FCameraShakePreviewerCommands::Register();
 
-	if (ensure(FModuleManager::Get().IsModuleLoaded(LevelEditorModuleName)))
+	if (ensure(FModuleManager::Get().IsModuleLoaded(UE::CameraShakePreviewer::Private::LevelEditorModuleName)))
 	{
 		RegisterEditorTab();
 		RegisterViewportOptionMenuExtender();
@@ -71,12 +95,15 @@ void FCameraShakePreviewerModule::ShutdownModule()
 
 void FCameraShakePreviewerModule::RegisterEditorTab()
 {
-	FLevelEditorModule& LevelEditorModule = FModuleManager::LoadModuleChecked<FLevelEditorModule>(LevelEditorModuleName);
+	FLevelEditorModule& LevelEditorModule =
+		FModuleManager::LoadModuleChecked<FLevelEditorModule>(UE::CameraShakePreviewer::Private::LevelEditorModuleName);
 
 	LevelEditorTabManagerChangedHandle = LevelEditorModule.OnTabManagerChanged().AddLambda([]()
 	{
 		// Add a new entry in the level editor's "Window" menu, which lets the user open the camera shake preview tool.
-		FLevelEditorModule& LevelEditorModule = FModuleManager::GetModuleChecked<FLevelEditorModule>(LevelEditorModuleName);
+			FLevelEditorModule& LevelEditorModule = FModuleManager::GetModuleChecked<FLevelEditorModule>(
+				UE::CameraShakePreviewer::Private::LevelEditorModuleName
+			);
 		TSharedPtr<FTabManager> LevelEditorTabManager = LevelEditorModule.GetLevelEditorTabManager();
 
 		const IWorkspaceMenuStructure& MenuStructure = WorkspaceMenu::GetMenuStructure();
@@ -94,78 +121,17 @@ void FCameraShakePreviewerModule::UnregisterEditorTab()
 {
 	if (LevelEditorTabManagerChangedHandle.IsValid())
 	{
-		FLevelEditorModule& LevelEditorModule = FModuleManager::GetModuleChecked<FLevelEditorModule>(LevelEditorModuleName);
+		FLevelEditorModule& LevelEditorModule =
+			FModuleManager::GetModuleChecked<FLevelEditorModule>(UE::CameraShakePreviewer::Private::LevelEditorModuleName);
 		LevelEditorModule.OnTabManagerChanged().Remove(LevelEditorTabManagerChangedHandle);
 	}
 }
 
 void FCameraShakePreviewerModule::RegisterViewportOptionMenuExtender()
 {
-	{
-		FToolMenuOwnerScoped ToolMenuOwnerScoped(this);
-
-		{
-			UToolMenu* Menu = UToolMenus::Get()->ExtendMenu("LevelEditor.LevelViewportToolBar.Options");
-
-			{
-				FToolMenuSection& Section = Menu->FindOrAddSection("LevelViewportViewportOptions2");
-
-				// ToggleCameraShakesPreview
-				{
-					static auto GetPerspectiveLevelEditorViewportClient = [](const FToolMenuContext& MenuContext) -> FLevelEditorViewportClient*
-					{
-						ULevelViewportToolBarContext* Context = MenuContext.FindContext<ULevelViewportToolBarContext>();
-						if (Context && Context->LevelViewportToolBarWidget.IsValid())
-						{
-							FLevelEditorViewportClient* ViewportClient = Context->GetLevelViewportClient();
-							if (ViewportClient && ViewportClient->ViewportType == ELevelViewportType::LVT_Perspective)
-							{
-								return ViewportClient;
-							}
-						}
-
-						return nullptr;
-					};
-
-					FToolUIAction Action;
-
-					Action.ExecuteAction.BindLambda([this](const FToolMenuContext& MenuContext)
-					{
-						if (FLevelEditorViewportClient* ViewportClient = GetPerspectiveLevelEditorViewportClient(MenuContext))
-						{
-							ToggleCameraShakesPreview(ViewportClient);
-						}
-					});
-
-					Action.GetActionCheckState.BindLambda([this](const FToolMenuContext& MenuContext) -> ECheckBoxState
-					{
-						if (FLevelEditorViewportClient* ViewportClient = GetPerspectiveLevelEditorViewportClient(MenuContext))
-						{
-							if (HasCameraShakesPreview(ViewportClient))
-							{
-								return ECheckBoxState::Checked;
-							}
-						}
-
-						return ECheckBoxState::Unchecked;
-					});
-
-					TSharedPtr<FUICommandInfo> ToggleCameraShakesPreview = FCameraShakePreviewerCommands::Get().ToggleCameraShakesPreview;
-					Section.AddMenuEntry(
-						"ToggleCameraShakesPreview",
-						ToggleCameraShakesPreview->GetLabel(),
-						ToggleCameraShakesPreview->GetDescription(),
-						FSlateIcon(FAppStyle::Get().GetStyleSetName(), "LevelViewport.ToggleCameraShakePreview"),
-						Action,
-						ToggleCameraShakesPreview->GetUserInterfaceType()
-					);
-				}
-			}
-		}
-	}
-
 	// Register a callback for adding a "Show Camera Shakes" option in the viewport options menu.
-	FLevelEditorModule& LevelEditorModule = FModuleManager::LoadModuleChecked<FLevelEditorModule>(LevelEditorModuleName);
+	FLevelEditorModule& LevelEditorModule =
+		FModuleManager::LoadModuleChecked<FLevelEditorModule>(UE::CameraShakePreviewer::Private::LevelEditorModuleName);
 
 	FLevelEditorModule::FLevelEditorMenuExtender Extender = FLevelEditorModule::FLevelEditorMenuExtender::CreateRaw(this, &FCameraShakePreviewerModule::OnExtendLevelViewportOptionMenu);
 	LevelEditorModule.GetAllLevelViewportOptionsMenuExtenders().Add(Extender);
@@ -180,7 +146,8 @@ void FCameraShakePreviewerModule::UnregisterViewportOptionMenuExtender()
 {
 	UToolMenus::UnregisterOwner(this);
 
-	FLevelEditorModule& LevelEditorModule = FModuleManager::GetModuleChecked<FLevelEditorModule>(LevelEditorModuleName);
+	FLevelEditorModule& LevelEditorModule =
+		FModuleManager::GetModuleChecked<FLevelEditorModule>(UE::CameraShakePreviewer::Private::LevelEditorModuleName);
 	LevelEditorModule.GetAllLevelViewportOptionsMenuExtenders().RemoveAll([this](const FLevelEditorModule::FLevelEditorMenuExtender& Extender) { return Extender.GetHandle() == ViewportOptionsMenuExtenderHandle; });
 
 	if (GEditor)
@@ -191,30 +158,24 @@ void FCameraShakePreviewerModule::UnregisterViewportOptionMenuExtender()
 
 TSharedRef<FExtender> FCameraShakePreviewerModule::OnExtendLevelViewportOptionMenu(const TSharedRef<FUICommandList> CommandList)
 {
-	// Find the viewport for which we're opening this options menu.
-	FLevelEditorModule& LevelEditorModule = FModuleManager::LoadModuleChecked<FLevelEditorModule>(LevelEditorModuleName);
-	TSharedPtr<ILevelEditor> LevelEditor = LevelEditorModule.GetLevelEditorInstance().Pin();
-	TSharedPtr<SLevelViewport> ViewportInterface = LevelEditor->GetActiveViewportInterface();
-	FLevelEditorViewportClient* ViewportClient = (FLevelEditorViewportClient*)&ViewportInterface->GetAssetViewportClient();
-	if (ViewportClient->ViewportType != ELevelViewportType::LVT_Perspective)
-	{
-		// If the current active viewport isn't perspective, let's not add any option for previewing camera shakes.
-		return MakeShared<FExtender>();
-	}
+	TSharedRef<FExtender> Extender(new FExtender());
+
+	// Legacy extension hook
+	Extender->AddMenuExtension(
+		"LevelViewportViewportOptions2",
+		EExtensionHook::After,
+		nullptr,
+		FMenuExtensionDelegate::CreateRaw(this, &FCameraShakePreviewerModule::CreateCameraShakeToggleOption)
+	);
+
+	// Adding separate extension hook for new toolbar
+	Extender->AddMenuExtension(
+		"CameraOptions",
+		EExtensionHook::After,
+		nullptr,
+		FMenuExtensionDelegate::CreateRaw(this, &FCameraShakePreviewerModule::CreateCameraShakeToggleOption)
+	);
 
-	// Bind the "toggle camera shakes preview" action to some callbacks that will set the flag on the appropriate viewport.
-	FUIAction ToggleCameraShakesPreviewAction;
-	ToggleCameraShakesPreviewAction.ExecuteAction.BindLambda([this, ViewportClient]()
-	{
-		ToggleCameraShakesPreview(ViewportClient);
-	});
-	ToggleCameraShakesPreviewAction.GetActionCheckState.BindLambda([this, ViewportClient]() -> ECheckBoxState
-	{
-		return HasCameraShakesPreview(ViewportClient) ? ECheckBoxState::Checked : ECheckBoxState::Unchecked;
-	});
-
-	// Make the extender that adds the menu entry.
-	TSharedRef<FExtender> Extender = MakeShared<FExtender>();
 	return Extender;
 }
 
@@ -239,6 +200,48 @@ void FCameraShakePreviewerModule::OnLevelViewportClientListChanged()
 	}
 }
 
+void FCameraShakePreviewerModule::CreateCameraShakeToggleOption(FMenuBuilder& InMenuBuilder)
+{
+	FUIAction ToggleCameraShakeAction;
+
+	ToggleCameraShakeAction.ExecuteAction.BindLambda(
+		[this]()
+		{
+			if (FLevelEditorViewportClient* ViewportClient =
+					UE::CameraShakePreviewer::Private::GetPerspectiveLevelEditorViewportClient())
+			{
+				ToggleCameraShakesPreview(ViewportClient);
+			}
+		}
+	);
+
+	ToggleCameraShakeAction.GetActionCheckState.BindLambda(
+		[this]() -> ECheckBoxState
+		{
+			if (FLevelEditorViewportClient* ViewportClient =
+					UE::CameraShakePreviewer::Private::GetPerspectiveLevelEditorViewportClient())
+			{
+				if (HasCameraShakesPreview(ViewportClient))
+				{
+					return ECheckBoxState::Checked;
+				}
+			}
+
+			return ECheckBoxState::Unchecked;
+		}
+	);
+
+	TSharedPtr<FUICommandInfo> ToggleCameraShakesPreview = FCameraShakePreviewerCommands::Get().ToggleCameraShakesPreview;
+	InMenuBuilder.AddMenuEntry(
+		ToggleCameraShakesPreview->GetLabel(),
+		ToggleCameraShakesPreview->GetDescription(),
+		FSlateIcon(FAppStyle::Get().GetStyleSetName(), "LevelViewport.ToggleCameraShakePreview"),
+		ToggleCameraShakeAction,
+		NAME_None,
+		EUserInterfaceActionType::ToggleButton
+	);
+}
+
 void FCameraShakePreviewerModule::ToggleCameraShakesPreview(FLevelEditorViewportClient* ViewportClient)
 {
 	if (FViewportInfo* ViewportInfo = ViewportInfos.Find(ViewportClient))
diff --git a/Engine/Plugins/Cameras/CameraShakePreviewer/Source/CameraShakePreviewer/Public/CameraShakePreviewerModule.h b/Engine/Plugins/Cameras/CameraShakePreviewer/Source/CameraShakePreviewer/Public/CameraShakePreviewerModule.h
--- a/Engine/Plugins/Cameras/CameraShakePreviewer/Source/CameraShakePreviewer/Public/CameraShakePreviewerModule.h
+++ b/Engine/Plugins/Cameras/CameraShakePreviewer/Source/CameraShakePreviewer/Public/CameraShakePreviewerModule.h
@@ -5,6 +5,7 @@
 #include "Modules/ModuleInterface.h"
 
 class FExtender;
+class FMenuBuilder;
 class FUICommandList;
 class FLevelEditorViewportClient;
 class FSpawnTabArgs;
@@ -50,6 +51,8 @@ private:
 	TSharedRef<FExtender> OnExtendLevelViewportOptionMenu(const TSharedRef<FUICommandList> CommandList);
 	void OnLevelViewportClientListChanged();
 
+	void CreateCameraShakeToggleOption(FMenuBuilder& InMenuBuilder);
+
 private:
 	FDelegateHandle LevelEditorTabManagerChangedHandle;
 	FDelegateHandle ViewportOptionsMenuExtenderHandle;
